[PCG] First draft for LowLevelsTests support for PCG

#rnx

diff --git a/Engine/Plugins/PCG/Tests/PCGTests.Build.cs b/Engine/Plugins/PCG/Tests/PCGTests.Build.cs
new file mode 100644
--- /dev/null
+++ b/Engine/Plugins/PCG/Tests/PCGTests.Build.cs
@@ -0,0 +1,26 @@
+// Copyright Epic Games, Inc. All Rights Reserved.
+using UnrealBuildTool;
+
+public class PCGTests : TestModuleRules
+{
+	static PCGTests()
+	{
+		if (InTestMode)
+		{
+			TestMetadata = new Metadata();
+			TestMetadata.TestName = "PCG";
+			TestMetadata.TestShortName = "PCG";
+		}
+	}
+	public PCGTests(ReadOnlyTargetRules Target) : base(Target, true)
+	{
+		bAllowUETypesInNamespaces = true;
+
+		PrivateDependencyModuleNames.AddRange(
+			new string[] {
+				"Core",
+				"CoreUObject",
+				"PCG"
+			});
+	}
+}
\ No newline at end of file
diff --git a/Engine/Plugins/PCG/Tests/PCGTests.Target.cs b/Engine/Plugins/PCG/Tests/PCGTests.Target.cs
new file mode 100644
--- /dev/null
+++ b/Engine/Plugins/PCG/Tests/PCGTests.Target.cs
@@ -0,0 +1,24 @@
+// Copyright Epic Games, Inc. All Rights Reserved.
+
+using UnrealBuildTool;
+
+[SupportedPlatforms(UnrealPlatformClass.All)]
+public class PCGTestsTarget : TestTargetRules
+{
+	public PCGTestsTarget(TargetInfo Target) : base(Target)
+	{
+		bCompileAgainstEngine = true;
+		bCompileAgainstCoreUObject = true;
+		bMockEngineDefaults = true;
+		bUsePlatformFileStub = true;
+
+		bTestsRequireEngine = true;
+
+		// If we add tests that are in the PCG Module (because they use functionality that is not exposed) we can enable this option.
+		// Note that it will gather ALL the tests that are in the engine by default, and if we execute the PCGTests program without any filter
+		// it will run them all.
+		// bWithLowLevelTestsOverride = true;
+        
+		bWarningsAsErrors = true;
+	}
+}
diff --git a/Engine/Plugins/PCG/Tests/Private/Init.cpp b/Engine/Plugins/PCG/Tests/Private/Init.cpp
new file mode 100644
--- /dev/null
+++ b/Engine/Plugins/PCG/Tests/Private/Init.cpp
@@ -0,0 +1,24 @@
+// Copyright Epic Games, Inc. All Rights Reserved.
+
+#include "CoreMinimal.h"
+
+#include "TestCommon/Initialization.h"
+
+#include "catch2/catch_tag_alias_autoregistrar.hpp"
+#include <catch2/catch_test_macros.hpp>
+
+//CATCH_REGISTER_TAG_ALIAS("[@PCGTests]", "[PCG]~[!benchmark]")
+
+GROUP_BEFORE_GLOBAL(Catch::DefaultGroup)
+{
+	// Find the engine directory before InitAll() enables the platform file stub. Prevents warnings.
+	FPlatformMisc::EngineDir();
+
+	InitAll(true, true);
+}
+
+GROUP_AFTER_GLOBAL(Catch::DefaultGroup)
+{
+	CleanupAll();
+	GIsRunning = false;
+}
\ No newline at end of file
diff --git a/Engine/Plugins/PCG/Tests/Private/SpatialAlgos/PCGOctreeQueriesTests.cpp b/Engine/Plugins/PCG/Tests/Private/SpatialAlgos/PCGOctreeQueriesTests.cpp
new file mode 100644
--- /dev/null
+++ b/Engine/Plugins/PCG/Tests/Private/SpatialAlgos/PCGOctreeQueriesTests.cpp
@@ -0,0 +1,76 @@
+﻿// Copyright Epic Games, Inc. All Rights Reserved.
+
+#include "CoreMinimal.h"
+#include "TestHarness.h"
+
+#include "PCGContext.h"
+#include "Data/PCGPointData.h"
+#include "SpatialAlgo/PCGOctreeQueries.h"
+
+/**
+* For all the tests, the point data will be points scattered on a straight line (Direction {1, 1, 1}), and evenly spaced.
+* Query will be on the first point, at the origin.
+*/
+TEST_CASE("PCG::OctreeQueries", "[PCG][OctreeQueries]")
+{
+	// Setup code for this test case
+	static constexpr double Distance = 100.0;
+	static constexpr int32 NumPoints = 10;
+	FPCGContext Context{};
+
+	UPCGBasePointData* InputPointData = FPCGContext::NewPointData_AnyThread(&Context);
+	InputPointData->SetNumPoints(NumPoints);
+	InputPointData->SetDensity(1.0f);
+	InputPointData->SetSeed(42);
+
+	TPCGValueRange<FTransform> TransformRange = InputPointData->GetTransformValueRange();
+
+	for (int32 i = 0; i < NumPoints; ++i)
+	{
+		FVector Location = FVector::OneVector * Distance * i;
+		TransformRange[i] = FTransform(Location);
+	}
+
+	SECTION("Sphere")
+	{
+		TArray<int32> ExpectedIndexes = { 0, 1, 2 };
+		int32 CountFound = 0;
+		UPCGOctreeQueries::ForEachPointInsideSphere(InputPointData, FVector::ZeroVector, 350.0, [&ExpectedIndexes, &CountFound](const UPCGBasePointData*, int32 PointIndex, double)
+		{
+			if (ExpectedIndexes.Contains(PointIndex))
+			{
+				++CountFound;
+			}
+		});
+
+		REQUIRE_EQUAL(CountFound, ExpectedIndexes.Num());
+	}
+
+	SECTION("ClosestPoint")
+	{
+		const int32 PointIndex = UPCGOctreeQueries::GetClosestPointIndex(InputPointData, FVector::ZeroVector, /*bDiscardCenter=*/false, 350.0);
+		REQUIRE_NOT_EQUAL(PointIndex, (int32)INDEX_NONE);
+		REQUIRE_EQUAL(PointIndex, 0);
+	}
+
+	SECTION("ClosestPointFromOtherPoint")
+	{
+		const int32 PointIndex = UPCGOctreeQueries::GetClosestPointIndexFromOtherPointIndex(InputPointData, 0, 350.0);
+		REQUIRE_NOT_EQUAL(PointIndex, (int32)INDEX_NONE);
+		REQUIRE_EQUAL(PointIndex, 1);
+	}
+
+	SECTION("FarthestPoint")
+	{
+		const int32 PointIndex = UPCGOctreeQueries::GetFarthestPointIndex(InputPointData, FVector::ZeroVector, 350.0);
+		REQUIRE_NOT_EQUAL(PointIndex, (int32)INDEX_NONE);
+		REQUIRE_EQUAL(PointIndex, 2);
+	}
+
+	SECTION("FarthestPointFromOtherPoint")
+	{
+		const int32 PointIndex = UPCGOctreeQueries::GetFarthestPointIndexFromOtherPointIndex(InputPointData, 0, 10000.0);
+		REQUIRE_NOT_EQUAL(PointIndex, (int32)INDEX_NONE);
+		REQUIRE_EQUAL(PointIndex, 9);
+	}
+}
