[PCG] Update Unit test with more sections

#rnx
#Synced-CL 49088187

diff --git a/Engine/Plugins/PCG/Tests/Private/Data/PCGBasePointDataTests.cpp b/Engine/Plugins/PCG/Tests/Private/Data/PCGBasePointDataTests.cpp
--- a/Engine/Plugins/PCG/Tests/Private/Data/PCGBasePointDataTests.cpp
+++ b/Engine/Plugins/PCG/Tests/Private/Data/PCGBasePointDataTests.cpp
@@ -82,101 +82,117 @@ TEST_CASE("PCG::BasePointData::API", "[PCG][BasePointData][API]")
 		FPCGPointOutputRange OutputRange{ PointData, RangeStartIndex, RangeSize };
 		const FPCGPointInputRange InputRange{ PointData, RangeStartIndex, RangeSize };
 
-		// Transform
-		TArray<FTransform> SetTransforms = PCGBasePointDataTests::CreateValueArray(PCGPoint.Transform, OutputRange.RangeSize);
-		PointData->SetTransformValuesOnRange(OutputRange, SetTransforms);
+		SECTION("Transform")
+		{
+			TArray<FTransform> SetTransforms = PCGBasePointDataTests::CreateValueArray(PCGPoint.Transform, OutputRange.RangeSize);
+			PointData->SetTransformValuesOnRange(OutputRange, SetTransforms);
 
-		TArray<FTransform> GetTransforms = PointData->GetTransformValuesFromRange(InputRange);
+			TArray<FTransform> GetTransforms = PointData->GetTransformValuesFromRange(InputRange);
 
-		PCGBasePointDataTests::TestArraysEqual(SetTransforms, GetTransforms);
+			PCGBasePointDataTests::TestArraysEqual(SetTransforms, GetTransforms);
 
-		const TConstPCGValueRange<FTransform> ConstTransformRange = PointData->GetConstTransformValueRange();
+			const TConstPCGValueRange<FTransform> ConstTransformRange = PointData->GetConstTransformValueRange();
 
-		PCGBasePointDataTests::TestRangeEqual(RangeStartIndex, ConstTransformRange, SetTransforms);
+			PCGBasePointDataTests::TestRangeEqual(RangeStartIndex, ConstTransformRange, SetTransforms);
+		}
 
-		// Density
-		TArray<float> SetDensity = PCGBasePointDataTests::CreateValueArray(PCGPoint.Density, OutputRange.RangeSize);
-		PointData->SetDensityValuesOnRange(OutputRange, SetDensity);
+		SECTION("Density")
+		{
+			TArray<float> SetDensity = PCGBasePointDataTests::CreateValueArray(PCGPoint.Density, OutputRange.RangeSize);
+			PointData->SetDensityValuesOnRange(OutputRange, SetDensity);
 
-		TArray<float> GetDensity = PointData->GetDensityValuesFromRange(InputRange);
+			TArray<float> GetDensity = PointData->GetDensityValuesFromRange(InputRange);
 
-		PCGBasePointDataTests::TestArraysEqual(SetDensity, GetDensity);
+			PCGBasePointDataTests::TestArraysEqual(SetDensity, GetDensity);
 
-		const TConstPCGValueRange<float> ConstDensityRange = PointData->GetConstDensityValueRange();
+			const TConstPCGValueRange<float> ConstDensityRange = PointData->GetConstDensityValueRange();
 
-		PCGBasePointDataTests::TestRangeEqual(RangeStartIndex, ConstDensityRange, SetDensity);
+			PCGBasePointDataTests::TestRangeEqual(RangeStartIndex, ConstDensityRange, SetDensity);
+		}
 
-		// Seed
-		TArray<int32> SetSeed = PCGBasePointDataTests::CreateValueArray(PCGPoint.Seed, OutputRange.RangeSize);
-		PointData->SetSeedValuesOnRange(OutputRange, SetSeed);
+		SECTION("Seed")
+		{
+			TArray<int32> SetSeed = PCGBasePointDataTests::CreateValueArray(PCGPoint.Seed, OutputRange.RangeSize);
+			PointData->SetSeedValuesOnRange(OutputRange, SetSeed);
 
-		TArray<int32> GetSeed = PointData->GetSeedValuesFromRange(InputRange);
+			TArray<int32> GetSeed = PointData->GetSeedValuesFromRange(InputRange);
 
-		PCGBasePointDataTests::TestArraysEqual(SetSeed, GetSeed);
+			PCGBasePointDataTests::TestArraysEqual(SetSeed, GetSeed);
 
-		const TConstPCGValueRange<int32> ConstSeedRange = PointData->GetConstSeedValueRange();
+			const TConstPCGValueRange<int32> ConstSeedRange = PointData->GetConstSeedValueRange();
 
-		PCGBasePointDataTests::TestRangeEqual(RangeStartIndex, ConstSeedRange, SetSeed);
+			PCGBasePointDataTests::TestRangeEqual(RangeStartIndex, ConstSeedRange, SetSeed);
+		}
 
-		// BoundsMin
-		TArray<FVector> SetBoundsMin = PCGBasePointDataTests::CreateValueArray(PCGPoint.BoundsMin, OutputRange.RangeSize);
-		PointData->SetBoundsMinValuesOnRange(OutputRange, SetBoundsMin);
+		SECTION("BoundsMin")
+		{
+			TArray<FVector> SetBoundsMin = PCGBasePointDataTests::CreateValueArray(PCGPoint.BoundsMin, OutputRange.RangeSize);
+			PointData->SetBoundsMinValuesOnRange(OutputRange, SetBoundsMin);
 
-		TArray<FVector> GetBoundsMin = PointData->GetBoundsMinValuesFromRange(InputRange);
+			TArray<FVector> GetBoundsMin = PointData->GetBoundsMinValuesFromRange(InputRange);
 
-		PCGBasePointDataTests::TestArraysEqual(SetBoundsMin, GetBoundsMin);
+			PCGBasePointDataTests::TestArraysEqual(SetBoundsMin, GetBoundsMin);
 
-		const TConstPCGValueRange<FVector> ConstBoundsMinRange = PointData->GetConstBoundsMinValueRange();
+			const TConstPCGValueRange<FVector> ConstBoundsMinRange = PointData->GetConstBoundsMinValueRange();
 
-		PCGBasePointDataTests::TestRangeEqual(RangeStartIndex, ConstBoundsMinRange, SetBoundsMin);
+			PCGBasePointDataTests::TestRangeEqual(RangeStartIndex, ConstBoundsMinRange, SetBoundsMin);
+		}
 
-		// BoundsMax
-		TArray<FVector> SetBoundsMax = PCGBasePointDataTests::CreateValueArray(PCGPoint.BoundsMax, OutputRange.RangeSize);
-		PointData->SetBoundsMaxValuesOnRange(OutputRange, SetBoundsMax);
+		SECTION("BoundsMax")
+		{
+			TArray<FVector> SetBoundsMax = PCGBasePointDataTests::CreateValueArray(PCGPoint.BoundsMax, OutputRange.RangeSize);
+			PointData->SetBoundsMaxValuesOnRange(OutputRange, SetBoundsMax);
 
-		TArray<FVector> GetBoundsMax = PointData->GetBoundsMaxValuesFromRange(InputRange);
+			TArray<FVector> GetBoundsMax = PointData->GetBoundsMaxValuesFromRange(InputRange);
 
-		PCGBasePointDataTests::TestArraysEqual(SetBoundsMax, GetBoundsMax);
+			PCGBasePointDataTests::TestArraysEqual(SetBoundsMax, GetBoundsMax);
 
-		const TConstPCGValueRange<FVector> ConstBoundsMaxRange = PointData->GetConstBoundsMaxValueRange();
+			const TConstPCGValueRange<FVector> ConstBoundsMaxRange = PointData->GetConstBoundsMaxValueRange();
 
-		PCGBasePointDataTests::TestRangeEqual(RangeStartIndex, ConstBoundsMaxRange, SetBoundsMax);
+			PCGBasePointDataTests::TestRangeEqual(RangeStartIndex, ConstBoundsMaxRange, SetBoundsMax);
+		}
 
-		// Color
-		TArray<FVector4> SetColor = PCGBasePointDataTests::CreateValueArray(PCGPoint.Color, OutputRange.RangeSize);
-		PointData->SetColorValuesOnRange(OutputRange, SetColor);
+		SECTION("Color")
+		{
+			TArray<FVector4> SetColor = PCGBasePointDataTests::CreateValueArray(PCGPoint.Color, OutputRange.RangeSize);
+			PointData->SetColorValuesOnRange(OutputRange, SetColor);
 
-		TArray<FVector4> GetColor = PointData->GetColorValuesFromRange(InputRange);
+			TArray<FVector4> GetColor = PointData->GetColorValuesFromRange(InputRange);
 
-		PCGBasePointDataTests::TestArraysEqual(SetColor, GetColor);
+			PCGBasePointDataTests::TestArraysEqual(SetColor, GetColor);
 
-		const TConstPCGValueRange<FVector4> ConstColorRange = PointData->GetConstColorValueRange();
+			const TConstPCGValueRange<FVector4> ConstColorRange = PointData->GetConstColorValueRange();
 
-		PCGBasePointDataTests::TestRangeEqual(RangeStartIndex, ConstColorRange, SetColor);
+			PCGBasePointDataTests::TestRangeEqual(RangeStartIndex, ConstColorRange, SetColor);
+		}
 
-		// Steepness
-		TArray<float> SetSteepness = PCGBasePointDataTests::CreateValueArray(PCGPoint.Steepness, OutputRange.RangeSize);
-		PointData->SetSteepnessValuesOnRange(OutputRange, SetSteepness);
+		SECTION("Steepness")
+		{
+			TArray<float> SetSteepness = PCGBasePointDataTests::CreateValueArray(PCGPoint.Steepness, OutputRange.RangeSize);
+			PointData->SetSteepnessValuesOnRange(OutputRange, SetSteepness);
 
-		TArray<float> GetSteepness = PointData->GetSteepnessValuesFromRange(InputRange);
+			TArray<float> GetSteepness = PointData->GetSteepnessValuesFromRange(InputRange);
 
-		PCGBasePointDataTests::TestArraysEqual(SetSteepness, GetSteepness);
+			PCGBasePointDataTests::TestArraysEqual(SetSteepness, GetSteepness);
 
-		const TConstPCGValueRange<float> ConstSteepnessRange = PointData->GetConstSteepnessValueRange();
+			const TConstPCGValueRange<float> ConstSteepnessRange = PointData->GetConstSteepnessValueRange();
 
-		PCGBasePointDataTests::TestRangeEqual(RangeStartIndex, ConstSteepnessRange, SetSteepness);
+			PCGBasePointDataTests::TestRangeEqual(RangeStartIndex, ConstSteepnessRange, SetSteepness);
+		}
 
-		// MetadataEntry
-		TArray<int64> SetMetadataEntry = PCGBasePointDataTests::CreateValueArray(PCGPoint.MetadataEntry, OutputRange.RangeSize);
-		PointData->SetMetadataEntryValuesOnRange(OutputRange, SetMetadataEntry);
+		SECTION("Metadata Entry")
+		{
+			TArray<int64> SetMetadataEntry = PCGBasePointDataTests::CreateValueArray(PCGPoint.MetadataEntry, OutputRange.RangeSize);
+			PointData->SetMetadataEntryValuesOnRange(OutputRange, SetMetadataEntry);
 
-		TArray<int64> GetMetadataEntry = PointData->GetMetadataEntryValuesFromRange(InputRange);
+			TArray<int64> GetMetadataEntry = PointData->GetMetadataEntryValuesFromRange(InputRange);
 
-		PCGBasePointDataTests::TestArraysEqual(SetMetadataEntry, GetMetadataEntry);
+			PCGBasePointDataTests::TestArraysEqual(SetMetadataEntry, GetMetadataEntry);
 
-		const TConstPCGValueRange<int64> ConstMetadataEntryRange = PointData->GetConstMetadataEntryValueRange();
+			const TConstPCGValueRange<int64> ConstMetadataEntryRange = PointData->GetConstMetadataEntryValueRange();
 
-		PCGBasePointDataTests::TestRangeEqual(RangeStartIndex, ConstMetadataEntryRange, SetMetadataEntry);
+			PCGBasePointDataTests::TestRangeEqual(RangeStartIndex, ConstMetadataEntryRange, SetMetadataEntry);
+		}
 	}
 
 	// Full Points
