[PCG - UnitTests] Silence warnings on test init coming from Slate/Streaming

#rnx

#Synced-CL 49088187

diff --git a/Engine/Plugins/PCG/Tests/Private/Init.cpp b/Engine/Plugins/PCG/Tests/Private/Init.cpp
--- a/Engine/Plugins/PCG/Tests/Private/Init.cpp
+++ b/Engine/Plugins/PCG/Tests/Private/Init.cpp
@@ -7,6 +7,10 @@
 #include "catch2/catch_tag_alias_autoregistrar.hpp"
 #include <catch2/catch_test_macros.hpp>
 
+#include "SlateGlobals.h"
+#include "Logging/LogScopedVerbosityOverride.h"
+#include "Styling/SlateWidgetStyleContainerBase.h"
+
 //CATCH_REGISTER_TAG_ALIAS("[@PCGTests]", "[PCG]~[!benchmark]")
 
 GROUP_BEFORE_GLOBAL(Catch::DefaultGroup)
@@ -14,6 +18,12 @@ GROUP_BEFORE_GLOBAL(Catch::DefaultGroup)
 	// Find the engine directory before InitAll() enables the platform file stub. Prevents warnings.
 	FPlatformMisc::EngineDir();
 
+	// Temporarily disable warnings during init
+	LOG_SCOPE_VERBOSITY_OVERRIDE(LogSlate, ELogVerbosity::Error);
+	LOG_SCOPE_VERBOSITY_OVERRIDE(LogSlateStyle, ELogVerbosity::Error);
+	LOG_SCOPE_VERBOSITY_OVERRIDE(LogStreaming, ELogVerbosity::Error);
+	LOG_SCOPE_VERBOSITY_OVERRIDE(LogUObjectGlobals, ELogVerbosity::Error);
+
 	InitAll(true, true);
 }
 
