[PCG] Add a base class for low-level PCG tests

#rnx
#Synced-CL 49044222

diff --git a/Engine/Plugins/PCG/Tests/Private/PCGTestsCommon.h b/Engine/Plugins/PCG/Tests/Private/PCGTestsCommon.h
new file mode 100644
--- /dev/null
+++ b/Engine/Plugins/PCG/Tests/Private/PCGTestsCommon.h
@@ -0,0 +1,23 @@
+﻿// Copyright Epic Games, Inc. All Rights Reserved.
+
+#include "CoreMinimal.h"
+#include "TestHarness.h"
+
+
+class FPCGBaseTest
+{
+public:
+	FPCGBaseTest()
+	{
+		PCGModule = MakeUnique<FPCGModule>();
+		PCGModule->StartupModule();
+	}
+	
+	virtual ~FPCGBaseTest()
+	{
+		PCGModule->ShutdownModule();
+	}
+	
+protected:
+	TUniquePtr<FPCGModule> PCGModule;
+};
\ No newline at end of file
