[PCG - UnitTests] Fix Linux build

#rnx
#Synced-CL 49122826

diff --git a/Engine/Plugins/PCG/Tests/PCGTests.Target.cs b/Engine/Plugins/PCG/Tests/PCGTests.Target.cs
--- a/Engine/Plugins/PCG/Tests/PCGTests.Target.cs
+++ b/Engine/Plugins/PCG/Tests/PCGTests.Target.cs
@@ -7,6 +7,7 @@ public class PCGTestsTarget : TestTargetRules
 {
 	public PCGTestsTarget(TargetInfo Target) : base(Target)
 	{
+		bCompileAgainstApplicationCore = true;
 		bCompileAgainstEngine = true;
 		bCompileAgainstCoreUObject = true;
 		bMockEngineDefaults = true;
diff --git a/Engine/Plugins/PCG/Tests/Private/Init.cpp b/Engine/Plugins/PCG/Tests/Private/Init.cpp
--- a/Engine/Plugins/PCG/Tests/Private/Init.cpp
+++ b/Engine/Plugins/PCG/Tests/Private/Init.cpp
@@ -10,6 +10,7 @@
 #include "SlateGlobals.h"
 #include "Logging/LogScopedVerbosityOverride.h"
 #include "Styling/SlateWidgetStyleContainerBase.h"
+#include "UObject/UObjectGlobals.h"
 
 //CATCH_REGISTER_TAG_ALIAS("[@PCGTests]", "[PCG]~[!benchmark]")
 
@@ -18,11 +19,13 @@ GROUP_BEFORE_GLOBAL(Catch::DefaultGroup)
 	// Find the engine directory before InitAll() enables the platform file stub. Prevents warnings.
 	FPlatformMisc::EngineDir();
 
-	// Temporarily disable warnings during init
+	// This will set the verbosity of the log temporarily while we init.
+	// It will skip the warnings for Slate and Streaming, while skipping error for UObjectGlobals
+	// as there are errors while contructing some CDOs about missing assets, which is irrelevant in our case.
 	LOG_SCOPE_VERBOSITY_OVERRIDE(LogSlate, ELogVerbosity::Error);
 	LOG_SCOPE_VERBOSITY_OVERRIDE(LogSlateStyle, ELogVerbosity::Error);
 	LOG_SCOPE_VERBOSITY_OVERRIDE(LogStreaming, ELogVerbosity::Error);
-	LOG_SCOPE_VERBOSITY_OVERRIDE(LogUObjectGlobals, ELogVerbosity::Error);
+	LOG_SCOPE_VERBOSITY_OVERRIDE(LogUObjectGlobals, ELogVerbosity::Fatal);
 
 	InitAll(true, true);
 }
@@ -31,4 +34,4 @@ GROUP_AFTER_GLOBAL(Catch::DefaultGroup)
 {
 	CleanupAll();
 	GIsRunning = false;
-}
\ No newline at end of file
+}
