[PCG - UnitTests] Small fixes to be able to load/shutdown module in unit tests

* Call PostInitEngine in StartupModule if we are already pass the post init engine
* Call PreExit in ShutdownModule if we didn't requested engine exit
* Both of those were tested in UnitTest and normal editor workflow to verify that the flow is the same in editor but works in unit tests too
* Will be cleaner for HotReload of PCG, as we were not kosher on this one
* Moved all ctor/dtor of PCGDataViewRegistry in the cpp as the TMap of TUniquePtr requires it.
* Moved PCGModule creation/startup/shutdown in PCGTestsCommon.cpp as TUniquePtr requires it.
* Add base class for a mock context

#rnx
#Synced-CL 49298804

diff --git a/Engine/Plugins/PCG/Source/PCG/Private/Data/Registry/PCGDataViewRegistry.cpp b/Engine/Plugins/PCG/Source/PCG/Private/Data/Registry/PCGDataViewRegistry.cpp
--- a/Engine/Plugins/PCG/Source/PCG/Private/Data/Registry/PCGDataViewRegistry.cpp
+++ b/Engine/Plugins/PCG/Source/PCG/Private/Data/Registry/PCGDataViewRegistry.cpp
@@ -5,6 +5,9 @@
 #include "Data/DataView/PCGDataView.h"
 #include "Data/DataView/PCGDataViewInterface.h"
 
+FPCGDataViewRegistry::FPCGDataViewRegistry() = default;
+FPCGDataViewRegistry::FPCGDataViewRegistry(FPCGDataViewRegistry&&) = default;
+FPCGDataViewRegistry& FPCGDataViewRegistry::operator=(FPCGDataViewRegistry&&) = default;
 FPCGDataViewRegistry::~FPCGDataViewRegistry() = default;
 
 void FPCGDataViewRegistry::RegisterPropertySelector(const TSubclassOf<UPCGData> InDataClass, TUniquePtr<const IPCGDataViewPropertySelector> PropertySelector)
diff --git a/Engine/Plugins/PCG/Source/PCG/Private/PCGModule.cpp b/Engine/Plugins/PCG/Source/PCG/Private/PCGModule.cpp
--- a/Engine/Plugins/PCG/Source/PCG/Private/PCGModule.cpp
+++ b/Engine/Plugins/PCG/Source/PCG/Private/PCGModule.cpp
@@ -210,14 +210,28 @@ void FPCGModule::StartupModule()
 
 	// Register onto the PreExit, because we need the class to be still valid to remove them from the mapping
 	FCoreDelegates::OnPreExit.AddRaw(this, &FPCGModule::PreExit);
-	// Also register on post init to gather all PCG types to register them
-	FCoreDelegates::OnPostEngineInit.AddRaw(this, &FPCGModule::OnPostInitEngine);
+
+	// Also register on post init to gather all PCG types to register them. Call it immediately if we already passed the PostInit phase (like in the UnitTests for example)
+	if (IPluginManager::Get().GetLastCompletedLoadingPhase() >= ELoadingPhase::PostEngineInit)
+	{
+		OnPostInitEngine();
+	}
+	else
+	{
+		FCoreDelegates::OnPostEngineInit.AddRaw(this, &FPCGModule::OnPostInitEngine);
+	}
 
 	TickDelegateHandle = FTSTicker::GetCoreTicker().AddTicker(FTickerDelegate::CreateRaw(this, &FPCGModule::Tick));
 }
 
 void FPCGModule::ShutdownModule()
 {
+	// Make sure to do the cleanup in PreExit if the module is shutdown while the engine still runs (like in the UnitTests for example)
+	if (!IsEngineExitRequested())
+	{
+		PreExit();
+	}
+
 	FTSTicker::RemoveTicker(TickDelegateHandle);
 
 	FCoreDelegates::OnPostEngineInit.RemoveAll(this);
diff --git a/Engine/Plugins/PCG/Source/PCG/Public/Data/Registry/PCGDataViewRegistry.h b/Engine/Plugins/PCG/Source/PCG/Public/Data/Registry/PCGDataViewRegistry.h
--- a/Engine/Plugins/PCG/Source/PCG/Public/Data/Registry/PCGDataViewRegistry.h
+++ b/Engine/Plugins/PCG/Source/PCG/Public/Data/Registry/PCGDataViewRegistry.h
@@ -12,11 +12,14 @@ struct FPCGDataView;
 class FPCGDataViewRegistry
 {
 public:
-	FPCGDataViewRegistry() = default;
+	FPCGDataViewRegistry();
+
+	FPCGDataViewRegistry(FPCGDataViewRegistry&&);
+	FPCGDataViewRegistry& operator=(FPCGDataViewRegistry&&);
+
 	FPCGDataViewRegistry(const FPCGDataViewRegistry&) = delete;
-	FPCGDataViewRegistry(FPCGDataViewRegistry&&) = default;
 	FPCGDataViewRegistry& operator=(const FPCGDataViewRegistry&) = delete;
-	FPCGDataViewRegistry& operator=(FPCGDataViewRegistry&&) = default;
+
 	virtual ~FPCGDataViewRegistry();
 
 	PCG_API void RegisterPropertySelector(const TSubclassOf<UPCGData> InDataClass, TUniquePtr<const IPCGDataViewPropertySelector> PropertySelector);
diff --git a/Engine/Plugins/PCG/Tests/Private/PCGTestsCommon.cpp b/Engine/Plugins/PCG/Tests/Private/PCGTestsCommon.cpp
new file mode 100644
--- /dev/null
+++ b/Engine/Plugins/PCG/Tests/Private/PCGTestsCommon.cpp
@@ -0,0 +1,24 @@
+﻿// Copyright Epic Games, Inc. All Rights Reserved.
+
+#include "PCGTestsCommon.h"
+
+#include "PCGModule.h"
+
+namespace PCGTests
+{
+	FPCGBaseTest::FPCGBaseTest()
+	{
+		PCGModule = MakeUnique<FPCGModule>();
+		PCGModule->StartupModule();
+	}
+
+	FPCGBaseTest::~FPCGBaseTest()
+	{
+		PCGModule->ShutdownModule();
+	}
+
+	UPCGBasePointData* FPCGBaseTestWithContext::CreatePointData()
+	{
+		return FPCGContext::NewPointData_AnyThread(&Context);
+	}
+}
\ No newline at end of file
diff --git a/Engine/Plugins/PCG/Tests/Private/PCGTestsCommon.h b/Engine/Plugins/PCG/Tests/Private/PCGTestsCommon.h
--- a/Engine/Plugins/PCG/Tests/Private/PCGTestsCommon.h
+++ b/Engine/Plugins/PCG/Tests/Private/PCGTestsCommon.h
@@ -1,23 +1,46 @@
 ﻿// Copyright Epic Games, Inc. All Rights Reserved.
 
+#include "PCGContext.h"
+
 #include "CoreMinimal.h"
 #include "TestHarness.h"
 
+class FPCGModule;
+class UPCGBasePointData;
 
+namespace PCGTests
+{
 class FPCGBaseTest
 {
 public:
-	FPCGBaseTest()
+	FPCGBaseTest();
+	virtual ~FPCGBaseTest();
+	
+protected:
+	TUniquePtr<FPCGModule> PCGModule;
+};
+
+/** Simple fixture that provide a FPCGContext, and helper functions to create new objects. */
+class FPCGBaseTestWithContext : public FPCGBaseTest
+{
+public:
+	using FPCGBaseTest::FPCGBaseTest;
+	
+protected:
+	UPCGBasePointData* CreatePointData();
+
+	template <typename T> requires std::is_base_of_v<UPCGData, T>
+	T* CreateData()
 	{
-		PCGModule = MakeUnique<FPCGModule>();
-		PCGModule->StartupModule();
+		return FPCGContext::NewObject_AnyThread<T>(&Context);
 	}
-	
-	virtual ~FPCGBaseTest()
+
+	template <typename T, typename ...Args> requires std::is_base_of_v<UObject, T>
+	T* CreateObject(Args&& ...InArgs)
 	{
-		PCGModule->ShutdownModule();
+		return FPCGContext::NewObject_AnyThread<T>(&Context, std::forward<Args>(InArgs)...);
 	}
 	
-protected:
-	TUniquePtr<FPCGModule> PCGModule;
-};
\ No newline at end of file
+	FPCGContext Context;
+};
+}
