[PCG] Add some common code for Low-level tests, and update Pathfinding tests accordingly

#rnx
#Synced-CL 49347717

diff --git a/Engine/Plugins/PCG/Tests/Private/Elements/PCGPathfindingTests.cpp b/Engine/Plugins/PCG/Tests/Private/Elements/PCGPathfindingTests.cpp
--- a/Engine/Plugins/PCG/Tests/Private/Elements/PCGPathfindingTests.cpp
+++ b/Engine/Plugins/PCG/Tests/Private/Elements/PCGPathfindingTests.cpp
@@ -10,25 +10,20 @@
 #include "Data/DataView/PCGDataViewInterface.h"
 #include "Data/PCGPointData.h"
 #include "Elements/PCGPathfindingElement.h"
-#include "Tests/PCGTestsCommon.h"
 #include "PCGTestsCommon.h"
+#include "Data/PCGPointArrayData.h"
 
-class FPCGPathfindingBaseTest: public FPCGBaseTest
+class FPCGPathfindingBaseTest: public PCGTests::FPCGSingleElementBaseTest<UPCGPathfindingSettings>
 {
 public:
 	FPCGPathfindingBaseTest(): 
-		FPCGBaseTest(),
+		FPCGSingleElementBaseTest(),
 		BaseStartPos(FRotator{}, FVector{ 0, 0, 0 }, FVector{}),
 		BaseGoalPos(FRotator{}, FVector{ 10, 10, 0 }, FVector{})
 	{
-		TestData.Reset(new PCGTestsCommon::FTestData());
-
-		PCGTestsCommon::GenerateSettings<UPCGPathfindingSettings>(*TestData);
-		Settings = CastChecked<UPCGPathfindingSettings>(TestData->Settings);
-		Settings->bOutputAsSpline = false;
-		Settings->bAcceptPartialPath = true;
-		PathfindingElement = TestData->Settings->GetElement();
-
+		TypedSettings->bOutputAsSpline = false;
+		TypedSettings->bAcceptPartialPath = true;
+		
 		TArray<FVector> Points;
 
 		Points.Add({ 0, 0, 0 });
@@ -61,7 +56,7 @@ public:
 		Points.Add({ 10, 8, 0 });
 		Points.Add({ 10, 9, 0 });
 
-		SimpleData = PCGTestsCommon::CreateEmptyBasePointData();
+		SimpleData = NewObject<UPCGPointArrayData>();
 		SimpleData->SetNumPoints(Points.Num());
 		SimpleData->SetDensity(1);
 		for (int i = 0; i < Points.Num(); i++)
@@ -72,11 +67,9 @@ public:
 
 	int GetPathfindingNumPoint()
 	{
-		Context = TestData->InitializeTestContext();
-
-		while (!PathfindingElement->Execute(Context.Get()))
-		{
-		}
+		ExecuteElement();
+		CHECK(NumErrors == 0);
+		CHECK(NumWarnings == 0);
 
 		const TArray<FPCGTaggedData>& Inputs = Context->InputData.GetAllSpatialInputs();
 		const TArray<FPCGTaggedData>& Outputs = Context->OutputData.GetAllSpatialInputs();
@@ -94,10 +87,6 @@ public:
 
 	int32 MiddleNodeIndex;
 
-	TUniquePtr<PCGTestsCommon::FTestData> TestData;
-	TUniquePtr<FPCGContext> Context;
-	UPCGPathfindingSettings* Settings;
-	FPCGElementPtr PathfindingElement;
 	FTransform BaseStartPos;
 	FTransform BaseGoalPos;
 
@@ -117,14 +106,14 @@ TEST_CASE_METHOD(FPCGPathfindingBaseTest, "PCG::PathfindingElement", "[PCG][Path
 	
 	DYNAMIC_SECTION("BaseTest with search distance = " << searchDistance)
 	{
-		FPCGTaggedData& SimpleTaggedData = TestData->InputData.TaggedData.Emplace_GetRef(FPCGTaggedData());
+		FPCGTaggedData& SimpleTaggedData = InputData.TaggedData.Emplace_GetRef(FPCGTaggedData());
 		SimpleTaggedData.Data = SimpleData;
 		SimpleTaggedData.Pin = PCGPinConstants::DefaultInputLabel;
 
-		Settings->Start = BaseStartPos.GetLocation();
-		Settings->Goal = BaseGoalPos.GetLocation();
+		TypedSettings->Start = BaseStartPos.GetLocation();
+		TypedSettings->Goal = BaseGoalPos.GetLocation();
 
-		Settings->SearchDistance = searchDistance;
+		TypedSettings->SearchDistance = searchDistance;
 		CHECK_EQUAL(GetPathfindingNumPoint(), expectedPathLength);
 	}
 }
@@ -141,12 +130,12 @@ TEST_CASE_METHOD(FPCGPathfindingBaseTest, "PCG::PathfindingElement:WithCosts", "
 	
 	DYNAMIC_SECTION("With Cost function = " << (WithCostFunction ? "true" : "false"))
 	{
-		FPCGTaggedData& SimpleTaggedData = TestData->InputData.TaggedData.Emplace_GetRef(FPCGTaggedData());
+		FPCGTaggedData& SimpleTaggedData = InputData.TaggedData.Emplace_GetRef(FPCGTaggedData());
 		SimpleTaggedData.Data = SimpleData;
 		SimpleTaggedData.Pin = PCGPinConstants::DefaultInputLabel;
 
-		Settings->Start = BaseStartPos.GetLocation();
-		Settings->Goal = BaseGoalPos.GetLocation();
+		TypedSettings->Start = BaseStartPos.GetLocation();
+		TypedSettings->Goal = BaseGoalPos.GetLocation();
 		
 		if (WithCostFunction)
 		{
@@ -159,36 +148,36 @@ TEST_CASE_METHOD(FPCGPathfindingBaseTest, "PCG::PathfindingElement:WithCosts", "
 			Attribute->SetValue(Key, 1000000.0);
 			SimpleData->GetMetadataEntryValueRange()[MiddleNodeIndex] = Key;
 
-			Settings->CostFunctionMode = EPCGPathfindingCostFunctionMode::CostMultiplier;
-			Settings->CostAttribute =  FPCGAttributePropertySelector::CreateSelectorFromString<FPCGAttributePropertyInputSelector>(CostAttributeName);			
+			TypedSettings->CostFunctionMode = EPCGPathfindingCostFunctionMode::CostMultiplier;
+			TypedSettings->CostAttribute =  FPCGAttributePropertySelector::CreateSelectorFromString<FPCGAttributePropertyInputSelector>(CostAttributeName);			
 		}
 
-		Settings->SearchDistance = 3.6;
+		TypedSettings->SearchDistance = 3.6;
 		CHECK_EQUAL(GetPathfindingNumPoint(), expectedPathLength);
 	}
 }
 
 TEST_CASE_METHOD(FPCGPathfindingBaseTest, "PCG::PathfindingElement:WithStartPosAsInput", "[PCG][PathfindingElement]")
 {
-	FPCGTaggedData& SimpleTaggedData = TestData->InputData.TaggedData.Emplace_GetRef(FPCGTaggedData());
+	FPCGTaggedData& SimpleTaggedData = InputData.TaggedData.Emplace_GetRef(FPCGTaggedData());
 	SimpleTaggedData.Data = SimpleData;
 	SimpleTaggedData.Pin = PCGPinConstants::DefaultInputLabel;
 
-	UPCGBasePointData* StartPointData = PCGTestsCommon::CreateEmptyBasePointData();
+	UPCGBasePointData* StartPointData = NewObject<UPCGPointArrayData>();
 	StartPointData->SetNumPoints(1);
 	StartPointData->SetDensity(1);
 	StartPointData->GetTransformValueRange()[0] = BaseStartPos;
 	
-	FPCGTaggedData& StartPointTaggedData = TestData->InputData.TaggedData.Emplace_GetRef(FPCGTaggedData());
+	FPCGTaggedData& StartPointTaggedData = InputData.TaggedData.Emplace_GetRef(FPCGTaggedData());
 	StartPointTaggedData.Data = StartPointData;
 	StartPointTaggedData.Pin = "Start";
-	Settings->bStartLocationsAsInput = true;
+	TypedSettings->bStartLocationsAsInput = true;
 
 	// Make sure the start pos is not used, so set it to an invalid value.
-	Settings->Start = FVector{ 1000, 1000, 1000 };
-	Settings->Goal = BaseGoalPos.GetLocation();	
+	TypedSettings->Start = FVector{ 1000, 1000, 1000 };
+	TypedSettings->Goal = BaseGoalPos.GetLocation();	
 	
-	Settings->SearchDistance = 3.6;
+	TypedSettings->SearchDistance = 3.6;
 	CHECK_EQUAL(GetPathfindingNumPoint(), 5);
 }
 
@@ -198,18 +187,18 @@ TEST_CASE_METHOD(FPCGPathfindingBaseTest, "PCG::PathfindingElement:WithStartPosA
 	
 	DYNAMIC_SECTION("With Copy From Orig = " << (WithCopyFromOriginatingPoints ? "true" : "false"))
 	{
-		FPCGTaggedData& SimpleTaggedData = TestData->InputData.TaggedData.Emplace_GetRef(FPCGTaggedData());
+		FPCGTaggedData& SimpleTaggedData = InputData.TaggedData.Emplace_GetRef(FPCGTaggedData());
 		SimpleTaggedData.Data = SimpleData;
 		SimpleTaggedData.Pin = PCGPinConstants::DefaultInputLabel;
 
 		const FString PosAttributeName = "ThePosition";
 
-		UPCGBasePointData* StartPointData = PCGTestsCommon::CreateEmptyBasePointData();
+		UPCGBasePointData* StartPointData = NewObject<UPCGPointArrayData>();
 		StartPointData->SetNumPoints(1);
 		StartPointData->SetDensity(1);
 		StartPointData->GetTransformValueRange()[0] = FTransform(FRotator{}, FVector{ 1000, 1000, 1000 }, FVector{});
 
-		FPCGTaggedData& StartPointTaggedData = TestData->InputData.TaggedData.Emplace_GetRef(FPCGTaggedData());
+		FPCGTaggedData& StartPointTaggedData = InputData.TaggedData.Emplace_GetRef(FPCGTaggedData());
 		StartPointTaggedData.Data = StartPointData;
 		StartPointTaggedData.Pin = "Start";
 		FPCGMetadataAttribute<FVector>* Attribute = StartPointData->MutableMetadata()->CreateAttribute<FVector>(PosAttributeName, FVector(), /*bAllowsInterpolation=*/true, /*bOverrideParent=*/false);
@@ -219,15 +208,15 @@ TEST_CASE_METHOD(FPCGPathfindingBaseTest, "PCG::PathfindingElement:WithStartPosA
 		Attribute->SetValue(Key, BaseStartPos.GetLocation());
 		StartPointData->GetMetadataEntryValueRange()[0] = Key;
 
-		Settings->bStartLocationsAsInput = true;
-		Settings->StartLocationAttribute = FPCGAttributePropertySelector::CreateSelectorFromString<FPCGAttributePropertyInputSelector>(PosAttributeName);
+		TypedSettings->bStartLocationsAsInput = true;
+		TypedSettings->StartLocationAttribute = FPCGAttributePropertySelector::CreateSelectorFromString<FPCGAttributePropertyInputSelector>(PosAttributeName);
 
 		// Make sure the start pos is not used, so set it to an invalid value.
-		Settings->Start = FVector{ 1000, 1000, 1000 };
-		Settings->Goal = BaseGoalPos.GetLocation();
+		TypedSettings->Start = FVector{ 1000, 1000, 1000 };
+		TypedSettings->Goal = BaseGoalPos.GetLocation();
 		
-		Settings->SearchDistance = 3.6;
-		Settings->bCopyOriginatingPoints = WithCopyFromOriginatingPoints;
+		TypedSettings->SearchDistance = 3.6;
+		TypedSettings->bCopyOriginatingPoints = WithCopyFromOriginatingPoints;
 		
 		CHECK_EQUAL(GetPathfindingNumPoint(), 5);
 		
@@ -249,6 +238,4 @@ TEST_CASE_METHOD(FPCGPathfindingBaseTest, "PCG::PathfindingElement:WithStartPosA
 			CHECK(OutputAttribute == nullptr);
 		}
 	}
-
-
 }
\ No newline at end of file
diff --git a/Engine/Plugins/PCG/Tests/Private/PCGTestsCommon.cpp b/Engine/Plugins/PCG/Tests/Private/PCGTestsCommon.cpp
--- a/Engine/Plugins/PCG/Tests/Private/PCGTestsCommon.cpp
+++ b/Engine/Plugins/PCG/Tests/Private/PCGTestsCommon.cpp
@@ -19,6 +19,31 @@ namespace PCGTests
 
 	UPCGBasePointData* FPCGBaseTestWithContext::CreatePointData()
 	{
-		return FPCGContext::NewPointData_AnyThread(&Context);
+		return FPCGContext::NewPointData_AnyThread(GetContext());
 	}
-}
\ No newline at end of file
+
+	FPCGTestsLogOutputDevice::FPCGTestsLogOutputDevice(): ThreadID(FPlatformTLS::GetCurrentThreadId())
+	{
+		GLog->AddOutputDevice(this);
+	}
+
+	FPCGTestsLogOutputDevice::~FPCGTestsLogOutputDevice()
+	{
+		GLog->RemoveOutputDevice(this);
+	}
+
+	void FPCGTestsLogOutputDevice::Serialize(const TCHAR* V, ELogVerbosity::Type Verbosity, const class FName& Category)
+	{
+		if (Verbosity <= ELogVerbosity::Warning)
+		{
+			if (NbMessageReceived.Contains(Verbosity))
+			{
+				NbMessageReceived[Verbosity]++;
+			}
+			else
+			{
+				NbMessageReceived.Add(Verbosity, 1);
+			}
+		}
+	}
+}
diff --git a/Engine/Plugins/PCG/Tests/Private/PCGTestsCommon.h b/Engine/Plugins/PCG/Tests/Private/PCGTestsCommon.h
--- a/Engine/Plugins/PCG/Tests/Private/PCGTestsCommon.h
+++ b/Engine/Plugins/PCG/Tests/Private/PCGTestsCommon.h
@@ -5,9 +5,61 @@
 #include "CoreMinimal.h"
 #include "TestHarness.h"
 
+#include <catch2/catch_approx.hpp>
+#include <catch2/matchers/catch_matchers_all.hpp>
+
 class FPCGModule;
 class UPCGBasePointData;
 
+namespace PCGTests
+{
+	template<typename VecType>
+	struct EqualsVecMatcher : Catch::Matchers::MatcherGenericBase {
+		EqualsVecMatcher(VecType const& Vec):
+			Vec{ Vec }
+		{}
+
+		template<typename OtherVecType>
+		bool match(OtherVecType const& Other) const {
+			return Vec.Equals(Other);
+		}
+
+		virtual std::string describe() const override {
+			return std::string("Equals: ") + std::string(TCHAR_TO_UTF8(*(Vec.ToString())));
+		}
+
+	private:
+		const VecType& Vec;
+	};
+
+	template <typename T> 
+	concept CHasEqual = requires(const T& V1, const T& V2) {{V1.Equals(V2)} -> std::same_as<bool>;};
+	template <typename T> 
+	concept CHasToString = requires(const T& V) { { V.ToString() } -> std::same_as<FString>; };
+
+	template<typename DataType>
+	auto ApproxMatcher(const DataType& Val)
+	{
+		return Catch::Matchers::WithinRel(Val);
+	}
+
+	template<CHasEqual DataType>
+	auto ApproxMatcher(const DataType& Val)
+	{
+		return EqualsVecMatcher(Val);
+	}
+}
+
+namespace Catch {
+	template<PCGTests::CHasToString T>
+	struct StringMaker<T> {
+		static std::string convert( T const& V ) {
+			return std::string(TCHAR_TO_UTF8(*V.ToString()));
+		}
+	};
+}
+
+
 namespace PCGTests
 {
 class FPCGBaseTest
@@ -20,27 +72,103 @@ protected:
 	TUniquePtr<FPCGModule> PCGModule;
 };
 
-/** Simple fixture that provide a FPCGContext, and helper functions to create new objects. */
+/** Simple fixture that provide helper functions to create new objects with a FPCGContext. */
 class FPCGBaseTestWithContext : public FPCGBaseTest
 {
 public:
 	using FPCGBaseTest::FPCGBaseTest;
 	
 protected:
+	virtual FPCGContext* GetContext() = 0;;
 	UPCGBasePointData* CreatePointData();
 
 	template <typename T> requires std::is_base_of_v<UPCGData, T>
 	T* CreateData()
 	{
-		return FPCGContext::NewObject_AnyThread<T>(&Context);
+		return FPCGContext::NewObject_AnyThread<T>(GetContext());
 	}
 
 	template <typename T, typename ...Args> requires std::is_base_of_v<UObject, T>
 	T* CreateObject(Args&& ...InArgs)
 	{
-		return FPCGContext::NewObject_AnyThread<T>(&Context, std::forward<Args>(InArgs)...);
+		return FPCGContext::NewObject_AnyThread<T>(GetContext(), std::forward<Args>(InArgs)...);
 	}
 	
+};
+	
+/** Simple fixture that provide a non-dynamic FPCGContext. */
+class FPCGBaseTestWithBasicContext : public FPCGBaseTestWithContext
+{
+public:
+	using FPCGBaseTestWithContext::FPCGBaseTestWithContext;
+
+protected:
+	virtual FPCGContext* GetContext() override {return &Context;};
+
 	FPCGContext Context;
 };
+	
+/** Utility class to gather errors and warnings*/
+struct FPCGTestsLogOutputDevice : public FOutputDevice
+{
+	FPCGTestsLogOutputDevice();
+
+	virtual ~FPCGTestsLogOutputDevice();
+
+	// FOutputDevice
+	virtual bool IsMemoryOnly() const override { return true; }
+	virtual bool CanBeUsedOnMultipleThreads() const override { return true; }
+	virtual bool CanBeUsedOnAnyThread() const override { return true; }
+	virtual void Serialize(const TCHAR* V, ELogVerbosity::Type Verbosity, const class FName& Category) override;
+
+	using MessageMap = TMap<ELogVerbosity::Type, int32>;
+	MessageMap NbMessageReceived;
+	const uint32 ThreadID;
+};
+	
+/** Base templated class to execute a PCG element.*/
+template<typename PCGSettingsType>
+class FPCGSingleElementBaseTest: public PCGTests::FPCGBaseTestWithContext
+{
+public:
+	FPCGSingleElementBaseTest(): 
+		FPCGBaseTestWithContext()
+	{
+		TypedSettings = NewObject<PCGSettingsType>();
+		check(TypedSettings);
+
+		TypedSettings->Seed = 42;
+
+		InputData.TaggedData.Emplace_GetRef().Data = TypedSettings;
+		InputData.TaggedData.Last().Pin = FName(TEXT("Settings"));
+		
+		TestElement = TypedSettings->GetElement();
+	}
+
+	virtual FPCGContext* GetContext() override {return Context.Get();};
+
+	void ExecuteElement()
+	{
+		Context.Reset(TypedSettings->GetElement()->Initialize(FPCGInitializeElementParams(&InputData, nullptr, nullptr)));
+		check(Context);
+		Context->InitializeSettings();
+		Context->AsyncState.NumAvailableTasks = 1;
+		
+		PCGTests::FPCGTestsLogOutputDevice LogCapture;
+
+		while (!TestElement->Execute(Context.Get()))
+		{
+		}	
+	
+		NumErrors = LogCapture.NbMessageReceived.Contains(ELogVerbosity::Error) ? LogCapture.NbMessageReceived[ELogVerbosity::Error] : 0;
+		NumWarnings = LogCapture.NbMessageReceived.Contains(ELogVerbosity::Warning) ? LogCapture.NbMessageReceived[ELogVerbosity::Warning] : 0;
+	}
+
+	PCGSettingsType* TypedSettings;
+	FPCGElementPtr TestElement;
+	FPCGDataCollection InputData;
+	TUniquePtr<FPCGContext> Context;
+	int NumErrors = 0;
+	int NumWarnings = 0;
+};	
 }
