[PCG] Fix unit test

#rnx

#Synced-CL 49723159

diff --git a/Engine/Plugins/PCG/Tests/Private/Benchmarks/Accessors/PCGAttributeAccessorBenchmarks.cpp b/Engine/Plugins/PCG/Tests/Private/Benchmarks/Accessors/PCGAttributeAccessorBenchmarks.cpp
--- a/Engine/Plugins/PCG/Tests/Private/Benchmarks/Accessors/PCGAttributeAccessorBenchmarks.cpp
+++ b/Engine/Plugins/PCG/Tests/Private/Benchmarks/Accessors/PCGAttributeAccessorBenchmarks.cpp
@@ -16,7 +16,7 @@
 
 #include "TestHarness.h"
 
-TEST_CASE_METHOD(PCGTests::FPCGBaseTestWithContext, "PCG::AttributeAccessorBenchmarks", "[PCG][AttributeAccessors][!benchmark]")
+TEST_CASE_METHOD(PCGTests::FPCGBaseTestWithBasicContext, "PCG::AttributeAccessorBenchmarks", "[PCG][AttributeAccessors][!benchmark]")
 {
 	constexpr int32 Seed = 42;
 	static const FPCGAttributePropertySelector PositionSelector = FPCGAttributePropertySelector::CreatePointPropertySelector(EPCGPointProperties::Position);
