[PCG] Enhance the Select Grammar to be able to work on Spline Data and Data Domain

#rnx
#Synced-CL 49723159

diff --git a/Engine/Plugins/PCG/Source/PCG/Private/Elements/Grammar/PCGSelectGrammar.cpp b/Engine/Plugins/PCG/Source/PCG/Private/Elements/Grammar/PCGSelectGrammar.cpp
--- a/Engine/Plugins/PCG/Source/PCG/Private/Elements/Grammar/PCGSelectGrammar.cpp
+++ b/Engine/Plugins/PCG/Source/PCG/Private/Elements/Grammar/PCGSelectGrammar.cpp
@@ -5,8 +5,10 @@
 #include "PCGContext.h"
 #include "PCGParamData.h"
 #include "Data/PCGBasePointData.h"
+#include "Data/PCGSplineData.h"
 #include "Elements/Metadata/PCGMetadataElementCommon.h"
 #include "Helpers/PCGPropertyHelpers.h"
+#include "Metadata/Accessors/PCGCustomAccessor.h"
 
 #include "Algo/Count.h"
 
@@ -120,7 +122,7 @@ TArray<FPCGPinProperties> UPCGSelectGrammarSettings::InputPinProperties() const
 {
 	TArray<FPCGPinProperties> PinProperties;
 
-	FPCGPinProperties& InputPin = PinProperties.Emplace_GetRef(PCGPinConstants::DefaultInputLabel, EPCGDataType::Point);
+	FPCGPinProperties& InputPin = PinProperties.Emplace_GetRef(PCGPinConstants::DefaultInputLabel, FPCGDataTypeIdentifier::Construct(FPCGDataTypeInfoPoint::AsId(), FPCGDataTypeInfoSpline::AsId()));
 	InputPin.SetRequiredPin();
 #if WITH_EDITOR
 	InputPin.Tooltip = LOCTEXT("InputPinTooltip", "A grammar will be selected based on criteria for all points.");
@@ -141,7 +143,7 @@ TArray<FPCGPinProperties> UPCGSelectGrammarSettings::InputPinProperties() const
 TArray<FPCGPinProperties> UPCGSelectGrammarSettings::OutputPinProperties() const
 {
 	TArray<FPCGPinProperties> PinProperties;
-	PinProperties.Emplace(PCGPinConstants::DefaultOutputLabel, EPCGDataType::Point);
+	PinProperties.Emplace(PCGPinConstants::DefaultOutputLabel, FPCGDataTypeIdentifier::Construct(FPCGDataTypeInfoPoint::AsId(), FPCGDataTypeInfoSpline::AsId()));
 
 	return PinProperties;
 }
@@ -153,10 +155,10 @@ bool FPCGSelectGrammarElement::ExecuteInternal(FPCGContext* InContext) const
 	const UPCGSelectGrammarSettings* Settings = InContext->GetInputSettings<UPCGSelectGrammarSettings>();
 	check(Settings);
 
-	const TArray<FPCGTaggedData> PointInputs = InContext->InputData.GetInputsByPin(PCGPinConstants::DefaultInputLabel);
+	const TArray<FPCGTaggedData> Inputs = InContext->InputData.GetInputsByPin(PCGPinConstants::DefaultInputLabel);
 	const TArray<FPCGTaggedData> CriteriaDataInputs = InContext->InputData.GetInputsByPin(PCGSelectGrammar::Constants::CriteriaDataPinLabel);
 
-	if (PointInputs.IsEmpty())
+	if (Inputs.IsEmpty())
 	{
 		return true;
 	}
@@ -257,32 +259,44 @@ bool FPCGSelectGrammarElement::ExecuteInternal(FPCGContext* InContext) const
 	}
 
 	TArray<FPCGTaggedData>& Outputs = InContext->OutputData.TaggedData;
-	for (const FPCGTaggedData& PointInput : PointInputs)
+	for (const FPCGTaggedData& Input : Inputs)
 	{
-		const UPCGBasePointData* InputPointData = Cast<UPCGBasePointData>(PointInput.Data);
-		if (!InputPointData || InputPointData->IsEmpty())
+		const UPCGData* InputData = Input.Data;
+		if (!InputData || (!InputData->IsA<UPCGBasePointData>() && !InputData->IsA<UPCGSplineData>()))
 		{
 			continue;
 		}
 
-		UPCGData* OutputData = InputPointData->DuplicateData(InContext);
-		UPCGBasePointData* OutputPointData = CastChecked<UPCGBasePointData>(OutputData);
-		check(OutputPointData->Metadata)
-		Outputs.Emplace_GetRef(PointInput).Data = OutputData;
+		FPCGAttributePropertyInputSelector KeySelector = Settings->KeyAttribute.CopyAndFixLast(InputData);
+		TUniquePtr<const IPCGAttributeAccessor> InputKeyAccessor;
+		TUniquePtr<const IPCGAttributeAccessorKeys> InputKeys;
 
-		FPCGAttributePropertyInputSelector KeySelector = Settings->KeyAttribute.CopyAndFixLast(InputPointData);
-		TUniquePtr<const IPCGAttributeAccessor> InputKeyAccessor = PCGAttributeAccessorHelpers::CreateConstAccessor(InputPointData, KeySelector);
-		TUniquePtr<const IPCGAttributeAccessorKeys> InputKeys = PCGAttributeAccessorHelpers::CreateConstKeys(InputPointData, KeySelector);
+		if (Settings->bKeyAsAttribute)
+		{
+			InputKeyAccessor = PCGAttributeAccessorHelpers::CreateConstAccessor(InputData, KeySelector);
+			InputKeys = PCGAttributeAccessorHelpers::CreateConstKeys(InputData, KeySelector);
+		}
+		else
+		{
+			InputKeyAccessor = MakeUnique<const FPCGConstantValueAccessor<FName>>(Settings->Key);
+			// Dummy single key
+			InputKeys = MakeUnique<FPCGAttributeAccessorKeysSingleObjectPtr<void>>();
+		}
+		
 		if (!InputKeyAccessor || !InputKeys)
 		{
 			PCGLog::Metadata::LogFailToCreateAccessorError(KeySelector, InContext);
 			continue;
 		}
 
+		UPCGData* OutputData = InputData->DuplicateData(InContext);
+		Outputs.Emplace_GetRef(Input).Data = OutputData;
+
 		FPCGAttributePropertyOutputSelector OutputGrammarSelector = Settings->OutputGrammarAttribute.CopyAndFixSource(&KeySelector);
-		FPCGMetadataAttribute<FString>* OutputAttribute = OutputPointData->Metadata->FindOrCreateAttribute(OutputGrammarSelector.GetName(), FString{}, /*bAllowsInterpolation=*/false);
-		TUniquePtr<IPCGAttributeAccessor> WriteGrammarAccessor = PCGAttributeAccessorHelpers::CreateAccessor(OutputAttribute, OutputPointData->Metadata);
-		TUniquePtr<IPCGAttributeAccessorKeys> OutputKeys = PCGAttributeAccessorHelpers::CreateKeys(OutputPointData, OutputGrammarSelector);
+		FPCGMetadataDomain* MetadataDomain = OutputData->MutableMetadata()->GetMetadataDomainFromSelector(OutputGrammarSelector);
+		FPCGMetadataAttribute<FString>* OutputAttribute = MetadataDomain->FindOrCreateAttribute(OutputGrammarSelector.GetName(), FString{}, /*bAllowsInterpolation=*/false);
+		TUniquePtr<IPCGAttributeAccessor> WriteGrammarAccessor = PCGAttributeAccessorHelpers::CreateAccessor(OutputAttribute, MetadataDomain);
+		TUniquePtr<IPCGAttributeAccessorKeys> OutputKeys = PCGAttributeAccessorHelpers::CreateKeys(OutputData, OutputGrammarSelector);
 		if (!WriteGrammarAccessor || !OutputKeys)
 		{
 			PCGLog::Metadata::LogFailToCreateAccessorError(OutputGrammarSelector, InContext);
@@ -372,9 +386,9 @@ bool FPCGSelectGrammarElement::ExecuteInternal(FPCGContext* InContext) const
 		}
 		else
 		{
-			FPCGAttributePropertyInputSelector ComparedValueSelector = Settings->ComparedValueAttribute.CopyAndFixLast(InputPointData);
-			TUniquePtr<const IPCGAttributeAccessor> ComparedValuesAccessor = PCGAttributeAccessorHelpers::CreateConstAccessor(InputPointData, ComparedValueSelector);
-			TUniquePtr<const IPCGAttributeAccessorKeys> ComparedValuesKeys = PCGAttributeAccessorHelpers::CreateConstKeys(InputPointData, ComparedValueSelector);
+			FPCGAttributePropertyInputSelector ComparedValueSelector = Settings->ComparedValueAttribute.CopyAndFixLast(InputData);
+			TUniquePtr<const IPCGAttributeAccessor> ComparedValuesAccessor = PCGAttributeAccessorHelpers::CreateConstAccessor(InputData, ComparedValueSelector);
+			TUniquePtr<const IPCGAttributeAccessorKeys> ComparedValuesKeys = PCGAttributeAccessorHelpers::CreateConstKeys(InputData, ComparedValueSelector);
 
 			if (!ComparedValuesAccessor || !ComparedValuesKeys)
 			{
@@ -383,9 +397,12 @@ bool FPCGSelectGrammarElement::ExecuteInternal(FPCGContext* InContext) const
 			}
 
 			// TODO: For now, cast everything to double for numerical comparison, but this could easily expand in the future
-			TArray<const IPCGAttributeAccessor*> Accessors = {InputKeyAccessor.Get(), ComparedValuesAccessor.Get()};
-			TArray<const IPCGAttributeAccessorKeys*> Keys = {InputKeys.Get(), ComparedValuesKeys.Get()};
-			PCGMetadataElementCommon::ApplyOnMultiAccessorsRange<FName, double>(Keys, Accessors, ProcessWithCriteria);
+			TStaticArray<const IPCGAttributeAccessor*, 2> Accessors = {InputKeyAccessor.Get(), ComparedValuesAccessor.Get()};
+			TStaticArray<const IPCGAttributeAccessorKeys*, 2> Keys = {InputKeys.Get(), ComparedValuesKeys.Get()};
+
+			// Technically we can mix Data and Elements domains, but if the output grammar goes into the data domain, we don't need to process more than one. (Effectively evaluating only the first element).
+			const int32 Count = OutputKeys->GetNum();
+			PCGMetadataElementCommon::ApplyOnMultiAccessorsRange<FName, double>(Keys, Accessors, ProcessWithCriteria, EPCGAttributeAccessorFlags::AllowBroadcastAndConstructible, PCGMetadataElementCommon::DefaultChunkSize, Count);
 		}
 	}
 
diff --git a/Engine/Plugins/PCG/Source/PCG/Public/Elements/Grammar/PCGSelectGrammar.h b/Engine/Plugins/PCG/Source/PCG/Public/Elements/Grammar/PCGSelectGrammar.h
--- a/Engine/Plugins/PCG/Source/PCG/Public/Elements/Grammar/PCGSelectGrammar.h
+++ b/Engine/Plugins/PCG/Source/PCG/Public/Elements/Grammar/PCGSelectGrammar.h
@@ -91,6 +91,7 @@ public:
 	virtual FText GetNodeTooltipText() const override;
 	virtual EPCGSettingsType GetType() const override { return EPCGSettingsType::Generic; }
 #endif
+	virtual bool HasDynamicPins() const override { return true; }
 	virtual void PostLoad() override;
 
 protected:
@@ -104,11 +105,11 @@ public:
 	UPROPERTY(BlueprintReadWrite, EditAnywhere, Category = Settings, meta = (PCG_Overridable))
 	bool bKeyAsAttribute = true;
 
-	/** An attribute key that represents the desired set of grammars. */
+	/** Each criterion has an associated key, this node will only evaluate (in order) the criteria that match this value. */
 	UPROPERTY(BlueprintReadWrite, EditAnywhere, Category = Settings, meta = (PCG_Overridable, EditCondition = "!bKeyAsAttribute", EditConditionHides))
 	FName Key;
 
-	/** An attribute selector for a key that represents the desired set of grammars. */
+	/** Attribute selector to read a key attribute from the input data. Each criterion has an associated key, this node will only evaluate (in order) the criteria that match the value of this attribute for each element. */
 	UPROPERTY(BlueprintReadWrite, EditAnywhere, Category = Settings, meta = (PCG_Overridable, EditCondition = "bKeyAsAttribute", EditConditionHides))
 	FPCGAttributePropertyInputSelector KeyAttribute;
 
diff --git a/Engine/Plugins/PCG/Tests/Private/Elements/Grammar/PCGSelectGrammarTests.cpp b/Engine/Plugins/PCG/Tests/Private/Elements/Grammar/PCGSelectGrammarTests.cpp
new file mode 100644
--- /dev/null
+++ b/Engine/Plugins/PCG/Tests/Private/Elements/Grammar/PCGSelectGrammarTests.cpp
@@ -0,0 +1,247 @@
+﻿// Copyright Epic Games, Inc. All Rights Reserved.
+
+#include "CoreMinimal.h"
+#include "TestHarness.h"
+
+#include <catch2/generators/catch_generators.hpp>
+
+#include "PCGContext.h"
+#include "Data/PCGPointData.h"
+#include "Data/PCGSplineData.h"
+#include "Elements/Grammar/PCGSelectGrammar.h"
+#include "Tests/PCGTestsCommon.h"
+#include "PCGTestsCommon.h"
+
+
+class FPCGSelectGrammarBaseTest: public PCGTests::FPCGBaseTest
+{
+public:
+	FPCGSelectGrammarBaseTest()
+		: PCGTests::FPCGBaseTest()
+	{
+		PCGTestsCommon::GenerateSettings<UPCGSelectGrammarSettings>(TestData);
+		Settings = CastChecked<UPCGSelectGrammarSettings>(TestData.Settings);
+		Element = TestData.Settings->GetElement();
+	}
+
+	template <typename PCGDataType> requires std::is_base_of_v<UPCGData, PCGDataType>
+	TArray<TTuple<const PCGDataType*, const FPCGMetadataAttribute<FString>*>> Execute(const FPCGAttributeIdentifier& GrammarAttributeId)
+	{
+		Context = TestData.InitializeTestContext();
+
+		while (!Element->Execute(Context.Get()))
+		{
+		}
+
+		const TArray<FPCGTaggedData>& Inputs = Context->InputData.GetInputsByPin(PCGPinConstants::DefaultInputLabel);
+		const TArray<FPCGTaggedData>& Outputs = Context->OutputData.GetInputsByPin(PCGPinConstants::DefaultOutputLabel);
+
+		CHECK_EQUAL(Inputs.Num(), Outputs.Num());
+
+		TArray<TTuple<const PCGDataType*, const FPCGMetadataAttribute<FString>*>> Result;
+		for (const FPCGTaggedData& OutputData : Outputs)
+		{
+			const PCGDataType* Data = Cast<const PCGDataType>(OutputData.Data);
+			CHECK_NOT_EQUAL(Data, nullptr);
+			const FPCGMetadataAttribute<FString>* Attribute = Data->ConstMetadata()->template GetConstTypedAttribute<FString>(GrammarAttributeId);
+			CHECK_NOT_EQUAL(Attribute, nullptr);
+			Result.Emplace(Data, Attribute);
+		}
+
+		return Result;
+	};
+
+	void EmplaceData(const UPCGData* Data, const FName PinLabel, bool bResetData)
+	{
+		if (bResetData)
+		{
+			TestData.Reset(Settings);
+		}
+
+		TestData.InputData.TaggedData.Emplace_GetRef(Data).Pin = PinLabel;
+	}
+
+	PCGTestsCommon::FTestData TestData;
+	TUniquePtr<FPCGContext> Context;
+	UPCGSelectGrammarSettings* Settings;
+	FPCGElementPtr Element;
+};
+
+TEST_CASE_METHOD(FPCGSelectGrammarBaseTest, "PCG::SelectGrammar", "[PCG][SelectGrammar]")
+{
+	const FName GrammarAttributeName = "Grammar";
+	const FName KeyAttributeName = "Key";
+	const FName DoubleAttributeName = "Double";
+	const FName Keys[2] = {"KeyA", "KeyB"};
+	const FString Grammars[3] = {TEXT("[A,B]"), TEXT("[A]*"), TEXT("{A,B}*")};
+
+	Settings->bCriteriaAsInput = false;
+	Settings->Criteria.Add(FPCGSelectGrammarCriterion{Keys[0], EPCGSelectGrammarComparator::LessThan, 250.0, 0.0, Grammars[0]});
+	Settings->Criteria.Add(FPCGSelectGrammarCriterion{Keys[0], EPCGSelectGrammarComparator::LessThan, 450.0, 0.0, Grammars[1]});
+	Settings->Criteria.Add(FPCGSelectGrammarCriterion{Keys[0], EPCGSelectGrammarComparator::Select, 0.0, 0.0, Grammars[2]});
+	Settings->Criteria.Add(FPCGSelectGrammarCriterion{Keys[1], EPCGSelectGrammarComparator::LessThan, 150.0, 0.0, Grammars[1]});
+	Settings->Criteria.Add(FPCGSelectGrammarCriterion{Keys[1], EPCGSelectGrammarComparator::Select, 0.0, 0.0, Grammars[0]});
+	
+	SECTION("PointData")
+	{
+		constexpr int32 NumPoints = 5;
+		
+		UPCGBasePointData* InputData = FPCGContext::NewPointData_AnyThread(nullptr);
+		InputData->AllocateProperties(EPCGPointNativeProperties::Transform | EPCGPointNativeProperties::MetadataEntry | EPCGPointNativeProperties::Density);
+		InputData->SetNumPoints(NumPoints);
+		
+		InputData->SetBoundsMin(FVector(-50, -50, -50));
+		InputData->SetBoundsMax(FVector(50, 50, 50));
+		TPCGValueRange<FTransform> TransformsRange = InputData->GetTransformValueRange();
+		TPCGValueRange<float> DensityRange = InputData->GetDensityValueRange();
+		TPCGValueRange<PCGMetadataEntryKey> MetadataEntryRange = InputData->GetMetadataEntryValueRange();
+
+		FPCGMetadataAttribute<FName>* KeyAttribute = InputData->MutableMetadata()->CreateAttribute<FName>(KeyAttributeName, NAME_None, false, true);
+		FPCGMetadataAttribute<FName>* KeyAttributeData = InputData->MutableMetadata()->CreateAttribute<FName>(FPCGAttributeIdentifier(KeyAttributeName, PCGMetadataDomainID::Data), Keys[0], false, true);
+		FPCGMetadataAttribute<double>* DoubleAttributeData = InputData->MutableMetadata()->CreateAttribute<double>(FPCGAttributeIdentifier(DoubleAttributeName, PCGMetadataDomainID::Data), 150.0, false, true);
+
+		for (int32 i = 0; i < NumPoints; ++i)
+		{
+			TransformsRange[i].SetScale3D(FVector::OneVector * (i + 1));
+			DensityRange[i] = i;
+			MetadataEntryRange[i] = InputData->MutableMetadata()->AddEntry();
+			KeyAttribute->SetValue(MetadataEntryRange[i], Keys[i % std::size(Keys)]);
+		}
+
+		EmplaceData(InputData, PCGPinConstants::DefaultInputLabel, /*bResetData=*/true);
+		
+		SECTION("Hardcoded Key: KeyA ; Comparator: $ScaledLocalSize.X")
+		{
+			Settings->bKeyAsAttribute = false;
+			Settings->Key = Keys[0];
+			Settings->ComparedValueAttribute.Update(TEXT("$ScaledLocalSize.X"));
+			Settings->OutputGrammarAttribute.SetAttributeName(GrammarAttributeName);
+
+			TArray<const FString*> ExpectedGrammars = {&Grammars[0], &Grammars[0], &Grammars[1], &Grammars[1], &Grammars[2]};
+
+			TArray<TTuple<const UPCGBasePointData*, const FPCGMetadataAttribute<FString>*>> DataAndAttributes = Execute<UPCGBasePointData>(FPCGAttributeIdentifier{GrammarAttributeName});
+			CHECK_EQUAL(DataAndAttributes.Num(), 1);
+			const UPCGBasePointData* OutputData = DataAndAttributes[0].Get<0>();
+			const FPCGMetadataAttribute<FString>* Attribute = DataAndAttributes[0].Get<1>();
+			CHECK(OutputData);
+			CHECK(Attribute);
+
+			for (int32 i = 0; i < NumPoints; ++i)
+			{
+				const FString ActualGrammar = Attribute->GetValueFromItemKey(OutputData->GetMetadataEntry(i));
+				REQUIRE_EQUAL(ActualGrammar, *ExpectedGrammars[i]);
+			}
+		}
+
+		SECTION("Attribute Key ; Comparator: $ScaledLocalSize.X")
+		{
+			Settings->bKeyAsAttribute = true;
+			Settings->KeyAttribute.SetAttributeName(KeyAttributeName); 
+			Settings->ComparedValueAttribute.Update(TEXT("$ScaledLocalSize.X"));
+			Settings->OutputGrammarAttribute.SetAttributeName(GrammarAttributeName);
+
+			TArray<const FString*> ExpectedGrammars = {&Grammars[0], &Grammars[0], &Grammars[1], &Grammars[0], &Grammars[2]};
+
+			TArray<TTuple<const UPCGBasePointData*, const FPCGMetadataAttribute<FString>*>> DataAndAttributes = Execute<UPCGBasePointData>(FPCGAttributeIdentifier{GrammarAttributeName});
+			CHECK_EQUAL(DataAndAttributes.Num(), 1);
+			const UPCGBasePointData* OutputData = DataAndAttributes[0].Get<0>();
+			const FPCGMetadataAttribute<FString>* Attribute = DataAndAttributes[0].Get<1>();
+			CHECK(OutputData);
+			CHECK(Attribute);
+
+			for (int32 i = 0; i < NumPoints; ++i)
+			{
+				const FString ActualGrammar = Attribute->GetValueFromItemKey(OutputData->GetMetadataEntry(i));
+				REQUIRE_EQUAL(ActualGrammar, *ExpectedGrammars[i]);
+			}
+		}
+
+		SECTION("Attribute Key ; Comparator: $Density")
+		{
+			Settings->bKeyAsAttribute = true;
+			Settings->KeyAttribute.SetAttributeName(KeyAttributeName); 
+			Settings->ComparedValueAttribute.SetPointProperty(EPCGPointProperties::Density);
+			Settings->OutputGrammarAttribute.SetAttributeName(GrammarAttributeName);
+
+			TArray<const FString*> ExpectedGrammars = {&Grammars[0], &Grammars[1], &Grammars[0], &Grammars[1], &Grammars[0]};
+
+			TArray<TTuple<const UPCGBasePointData*, const FPCGMetadataAttribute<FString>*>> DataAndAttributes = Execute<UPCGBasePointData>(FPCGAttributeIdentifier{GrammarAttributeName});
+			CHECK_EQUAL(DataAndAttributes.Num(), 1);
+			const UPCGBasePointData* OutputData = DataAndAttributes[0].Get<0>();
+			const FPCGMetadataAttribute<FString>* Attribute = DataAndAttributes[0].Get<1>();
+			CHECK(OutputData);
+			CHECK(Attribute);
+
+			for (int32 i = 0; i < NumPoints; ++i)
+			{
+				const FString ActualGrammar = Attribute->GetValueFromItemKey(OutputData->GetMetadataEntry(i));
+				REQUIRE_EQUAL(ActualGrammar, *ExpectedGrammars[i]);
+			}
+		}
+
+		SECTION("Attribute Key ; Comparator: @Data.Double ; Output Grammar: @Data.Grammar")
+		{
+			Settings->bKeyAsAttribute = true;
+			Settings->KeyAttribute = FPCGAttributePropertyInputSelector::CreateAttributeSelector<FPCGAttributePropertyInputSelector>(KeyAttributeName, TEXT("Data"));
+			Settings->ComparedValueAttribute = FPCGAttributePropertyInputSelector::CreateAttributeSelector<FPCGAttributePropertyInputSelector>(DoubleAttributeName, TEXT("Data"));
+			Settings->OutputGrammarAttribute = FPCGAttributePropertyOutputSelector::CreateAttributeSelector<FPCGAttributePropertyOutputSelector>(GrammarAttributeName, TEXT("Data"));
+
+			TArray<TTuple<const UPCGBasePointData*, const FPCGMetadataAttribute<FString>*>> DataAndAttributes = Execute<UPCGBasePointData>(FPCGAttributeIdentifier{GrammarAttributeName, PCGMetadataDomainID::Data});
+			CHECK_EQUAL(DataAndAttributes.Num(), 1);
+			const FPCGMetadataAttribute<FString>* Attribute = DataAndAttributes[0].Get<1>();
+			CHECK(Attribute);
+
+			REQUIRE(Attribute->GetMetadataDomain()->GetDomainID() == PCGMetadataDomainID::Data);
+			REQUIRE_EQUAL(Attribute->GetValueFromItemKey(PCGFirstEntryKey), Grammars[0]);
+		}
+	}
+	
+	SECTION("SplineData")
+	{
+		constexpr int32 NumPoints = 5;
+
+		TArray<FSplinePoint> SplinePoints1;
+		TArray<FSplinePoint> SplinePoints2;
+		for (int32 i = 0; i < NumPoints; ++i)
+		{
+			SplinePoints1.Emplace(i, FVector::OneVector * i * 100);
+			SplinePoints2.Emplace(i, FVector::OneVector * i * -50);
+		}
+
+		auto PrepareData = [this, KeyAttributeName, DoubleAttributeName, &Keys](const TArray<FSplinePoint>& SplinePoints, int32 Index)
+		{
+			UPCGSplineData* InputData = NewObject<UPCGSplineData>();
+			InputData->Initialize(SplinePoints, /*bClosedLoop=*/false, FTransform::Identity);
+
+			InputData->MutableMetadata()->CreateAttribute<FName>(FPCGAttributeIdentifier(KeyAttributeName, PCGMetadataDomainID::Data), Keys[Index], false, true);
+			InputData->MutableMetadata()->CreateAttribute<double>(FPCGAttributeIdentifier(DoubleAttributeName, PCGMetadataDomainID::Data), InputData->GetLength(), false, true);
+
+			EmplaceData(InputData, PCGPinConstants::DefaultInputLabel, /*bResetData=*/Index==0);
+		};
+
+		TestData.InputData.TaggedData.Reset();
+		PrepareData(SplinePoints1, 0);
+		PrepareData(SplinePoints2, 1);
+
+		SECTION("Attribute Key ; Comparator: @Data.Double ; Output Grammar: @Data.Grammar")
+		{
+			Settings->bKeyAsAttribute = true;
+			Settings->KeyAttribute = FPCGAttributePropertyInputSelector::CreateAttributeSelector<FPCGAttributePropertyInputSelector>(KeyAttributeName, TEXT("Data"));
+			Settings->ComparedValueAttribute = FPCGAttributePropertyInputSelector::CreateAttributeSelector<FPCGAttributePropertyInputSelector>(DoubleAttributeName, TEXT("Data"));
+			Settings->OutputGrammarAttribute = FPCGAttributePropertyOutputSelector::CreateAttributeSelector<FPCGAttributePropertyOutputSelector>(GrammarAttributeName, TEXT("Data"));
+
+			TArray<TTuple<const UPCGSplineData*, const FPCGMetadataAttribute<FString>*>> DataAndAttributes = Execute<UPCGSplineData>(FPCGAttributeIdentifier{GrammarAttributeName, PCGMetadataDomainID::Data});
+			CHECK_EQUAL(DataAndAttributes.Num(), 2);
+			const FPCGMetadataAttribute<FString>* Attribute1 = DataAndAttributes[0].Get<1>();
+			const FPCGMetadataAttribute<FString>* Attribute2 = DataAndAttributes[1].Get<1>();
+			CHECK(Attribute1);
+			CHECK(Attribute2);
+
+			REQUIRE(Attribute1->GetMetadataDomain()->GetDomainID() == PCGMetadataDomainID::Data);
+			REQUIRE_EQUAL(Attribute1->GetValueFromItemKey(PCGFirstEntryKey), Grammars[2]);
+
+			REQUIRE(Attribute2->GetMetadataDomain()->GetDomainID() == PCGMetadataDomainID::Data);
+			REQUIRE_EQUAL(Attribute2->GetValueFromItemKey(PCGFirstEntryKey), Grammars[0]);
+		}
+	}
+}
