[PCG] Add `[REDACTED_HANDLE].$Length` as a spline accessor

#rnx
#Synced-CL 49758339

diff --git a/Engine/Plugins/PCG/Source/PCG/Private/Data/PCGSplineData.cpp b/Engine/Plugins/PCG/Source/PCG/Private/Data/PCGSplineData.cpp
--- a/Engine/Plugins/PCG/Source/PCG/Private/Data/PCGSplineData.cpp
+++ b/Engine/Plugins/PCG/Source/PCG/Private/Data/PCGSplineData.cpp
@@ -427,11 +427,15 @@ TUniquePtr<IPCGAttributeAccessor> UPCGSplineData::CreateStaticAccessor(const FPC
 		
 		if (PropertyName == TEXT("SplineTransform"))
 		{
-			return MakeUnique<FPCGSplineAccessor<FTransform, EPCGSplineAccessorTarget::Transform>>(PCGSplineStructProperty);
+			return MakeUnique<FPCGSplineAccessorV2<FTransform, EPCGSplineDataProperties::SplineTransform>>(PCGSplineStructProperty);
 		}
 		else if (PropertyName == TEXT("IsClosed"))
 		{
-			return MakeUnique<FPCGSplineAccessor<bool, EPCGSplineAccessorTarget::ClosedLoop>>(PCGSplineStructProperty);
+			return MakeUnique<FPCGSplineAccessorV2<bool, EPCGSplineDataProperties::IsClosed>>(PCGSplineStructProperty);
+		}
+		else if (PropertyName == TEXT("Length"))
+		{
+			return MakeUnique<FPCGSplineAccessorV2<double, EPCGSplineDataProperties::Length>>(PCGSplineStructProperty);
 		}
 	}
 	
diff --git a/Engine/Plugins/PCG/Source/PCG/Public/Data/PCGSplineStruct.h b/Engine/Plugins/PCG/Source/PCG/Public/Data/PCGSplineStruct.h
--- a/Engine/Plugins/PCG/Source/PCG/Public/Data/PCGSplineStruct.h
+++ b/Engine/Plugins/PCG/Source/PCG/Public/Data/PCGSplineStruct.h
@@ -32,8 +32,9 @@ enum class EPCGSplineStructProperties : uint8
 UENUM()
 enum class EPCGSplineDataProperties : uint8
 {
-	SplineTransform	 UMETA(Tooltip = "Transform of the spline.", PCG_MetadataDomain="Data"),
-	IsClosed		 UMETA(Tooltip = "If the spline is closed.", PCG_MetadataDomain="Data")
+	SplineTransform  UMETA(Tooltip = "Transform of the spline.", PCG_MetadataDomain="Data"),
+	IsClosed         UMETA(Tooltip = "If the spline is closed.", PCG_MetadataDomain="Data"),
+	Length           UMETA(Tooltip = "Length of the spline", PCG_MetadataDomain="Data")
 };
 
 enum class EPCGControlPointsAccessorTarget
diff --git a/Engine/Plugins/PCG/Source/PCG/Public/Metadata/Accessors/PCGSplineAccessor.h b/Engine/Plugins/PCG/Source/PCG/Public/Metadata/Accessors/PCGSplineAccessor.h
--- a/Engine/Plugins/PCG/Source/PCG/Public/Metadata/Accessors/PCGSplineAccessor.h
+++ b/Engine/Plugins/PCG/Source/PCG/Public/Metadata/Accessors/PCGSplineAccessor.h
@@ -7,12 +7,6 @@
 
 #include "Components/SplineComponent.h"
 
-enum class EPCGSplineAccessorTarget
-{
-	Transform,
-	ClosedLoop
-};
-
 /**
 * Templated accessor for location/rotation/scale in world coordinates. It's important that the keys only have a single value, the struct that holds the spline data/struct, since
 * there is a single transform per spline.
@@ -251,19 +245,19 @@ private:
 };
 
 /**
-* Templated accessor for global spline data. Note that closed loop value is read-only.
+* Templated accessor for global spline data. Note that only the spline transform value is writable.
 * Keys supported: PCGSplineData, FPCGSplineStruct
 */
-template <typename T, EPCGSplineAccessorTarget Target>
-class FPCGSplineAccessor : public IPCGAttributeAccessorT<FPCGSplineAccessor<T, Target>>, IPCGPropertyChain
+template <typename T, EPCGSplineDataProperties Target>
+class FPCGSplineAccessorV2 : public IPCGAttributeAccessorT<FPCGSplineAccessorV2<T, Target>>, IPCGPropertyChain
 {
 public:
 	// The underlying type is a quat if we target the rotation, otherwise a vector
 	using Type = T;
-	using Super = IPCGAttributeAccessorT<FPCGSplineAccessor<T, Target>>;
+	using Super = IPCGAttributeAccessorT<FPCGSplineAccessorV2<T, Target>>;
 	
-	FPCGSplineAccessor(const FStructProperty* InProperty, TArray<const FProperty*>&& ExtraProperties = {})
-		: Super(/*bInReadOnly=*/ Target == EPCGSplineAccessorTarget::ClosedLoop)
+	FPCGSplineAccessorV2(const FStructProperty* InProperty, TArray<const FProperty*>&& ExtraProperties = {})
+		: Super(/*bInReadOnly=*/ Target != EPCGSplineDataProperties::SplineTransform)
 		, IPCGPropertyChain(InProperty, std::move(ExtraProperties))
 	{
 		static_assert(PCG::Private::IsPCGType<Type>());		
@@ -294,16 +288,24 @@ public:
 		
 		for (int32 i = 0; i < OutValues.Num(); ++i)
 		{
-			if constexpr (Target == EPCGSplineAccessorTarget::Transform)
+			if constexpr (Target == EPCGSplineDataProperties::SplineTransform)
 			{
-				static_assert(std::is_same_v<Type, FTransform>);
+				using SplineTransformType = decltype(SplineStruct.GetTransform());
+				static_assert(std::is_same_v<Type, SplineTransformType>);
 				OutValues[i] = SplineStruct.GetTransform();
 			}
-			else if constexpr (Target == EPCGSplineAccessorTarget::ClosedLoop)
+			else if constexpr (Target == EPCGSplineDataProperties::IsClosed)
 			{
-				static_assert(std::is_same_v<Type, bool>);
+				using SplineIsClosedType = decltype(SplineStruct.IsClosedLoop());
+				static_assert(std::is_same_v<Type, SplineIsClosedType>);
 				OutValues[i] = SplineStruct.IsClosedLoop();
 			}
+			else if constexpr (Target == EPCGSplineDataProperties::Length)
+			{
+				using SplineLengthType = decltype(SplineStruct.GetSplineLength());
+				static_assert(std::is_same_v<Type, SplineLengthType>);
+				OutValues[i] = SplineStruct.GetSplineLength();
+			}
 			else
 			{
 				// Pitfall static assert
@@ -316,53 +318,85 @@ public:
 
 	bool SetRangeImpl(TArrayView<const Type> InValues, int32 Index, IPCGAttributeAccessorKeys& Keys, EPCGAttributeAccessorFlags)
 	{
-		if (Target == EPCGSplineAccessorTarget::ClosedLoop)
+		if constexpr (Target != EPCGSplineDataProperties::SplineTransform)
 		{
 			// Not supported, as we need to update the spline when we change it, which is not threadsafe.
 			return false;
 		}
-			
-		// We want to access the transform on the spline struct, there is just a single one.
-		void* ContainerKeys = nullptr;
-		TArrayView<void*> ContainerKeysView(&ContainerKeys, 1);
-		if (!Keys.GetKeys(Index, ContainerKeysView))
-		{
-			return false;
-		}
-
-		// Validation to not access keys that are not the expected type. Done after the GetKeys, as we also want to discard other type of incompatible
-		// keys (like a Default Metadata entry key)
-		if (!ensure(Keys.IsClassSupported(TopPropertyStruct)))
+		else
 		{
-			return false;
-		}
-
-		PCGPropertyAccessor::AddressOffset(GetPropertyChain(), ContainerKeysView);
+			// We want to access the transform on the spline struct, there is just a single one.
+			void* ContainerKeys = nullptr;
+			TArrayView<void*> ContainerKeysView(&ContainerKeys, 1);
+			if (!Keys.GetKeys(Index, ContainerKeysView))
+			{
+				return false;
+			}
 
-		FPCGSplineStruct& SplineStruct = *static_cast<FPCGSplineStruct*>(ContainerKeys);
-		
-		for (int32 i = 0; i < InValues.Num(); ++i)
-		{
-			if constexpr (Target == EPCGSplineAccessorTarget::Transform)
+			// Validation to not access keys that are not the expected type. Done after the GetKeys, as we also want to discard other type of incompatible
+			// keys (like a Default Metadata entry key)
+			if (!ensure(Keys.IsClassSupported(TopPropertyStruct)))
 			{
-				static_assert(std::is_same_v<Type, FTransform>);
-				SplineStruct.Transform = InValues[i];
+				return false;
 			}
-			else if constexpr (Target == EPCGSplineAccessorTarget::ClosedLoop)
+
+			PCGPropertyAccessor::AddressOffset(GetPropertyChain(), ContainerKeysView);
+
+			FPCGSplineStruct& SplineStruct = *static_cast<FPCGSplineStruct*>(ContainerKeys);
+			
+			if constexpr (std::is_same_v<Type, decltype(SplineStruct.Transform)>)
 			{
-				// Not supported, as we need to update the spline when we change it, which is not threadsafe.
-				// Need to stay here to avoid unreachable code.
+				for (int32 i = 0; i < InValues.Num(); ++i)
+				{
+					SplineStruct.Transform = InValues[i];
+				}
 			}
 			else
 			{
 				// Pitfall static assert
 				static_assert(!std::is_same_v<Type, Type>);
 			}
+			
+			return true;
 		}
-
-		return true;
 	}
 
 private:
 	const UStruct* TopPropertyStruct = nullptr;
-};
\ No newline at end of file
+};
+
+// DEPRECATION
+
+enum class UE_DEPRECATED(5.8, "Not used anymore. Replaced by EPCGSplineDataProperties") EPCGSplineAccessorTarget
+{
+	Transform,
+	ClosedLoop
+};
+
+PRAGMA_DISABLE_DEPRECATION_WARNINGS
+// Consteval function to convert the old enum to the new one, to be used in a template (hence the consteval)
+template <EPCGSplineAccessorTarget Target>
+consteval EPCGSplineDataProperties DeprecationSplineAccessorTarget()
+{
+	if constexpr (Target == EPCGSplineAccessorTarget::Transform)
+	{
+		return EPCGSplineDataProperties::SplineTransform;
+	}
+	else if constexpr (Target == EPCGSplineAccessorTarget::ClosedLoop)
+	{
+		return EPCGSplineDataProperties::IsClosed;
+	}
+	else
+	{
+		// Pitfall assert
+		static_assert(Target != Target);
+		return {};
+	}
+}
+
+template <typename T, EPCGSplineAccessorTarget Target>
+class UE_DEPRECATED(5.8, "Use FPCGSplineAccessorV2") FPCGSplineAccessor : public FPCGSplineAccessorV2<T, DeprecationSplineAccessorTarget<Target>()>
+{
+	using FPCGSplineAccessorV2<T, DeprecationSplineAccessorTarget<Target>()>::FPCGSplineAccessorV2;
+};
+PRAGMA_ENABLE_DEPRECATION_WARNINGS
diff --git a/Engine/Plugins/PCG/Tests/Private/Metadata/Accessors/PCGSplineAccessorTests.cpp b/Engine/Plugins/PCG/Tests/Private/Metadata/Accessors/PCGSplineAccessorTests.cpp
new file mode 100644
--- /dev/null
+++ b/Engine/Plugins/PCG/Tests/Private/Metadata/Accessors/PCGSplineAccessorTests.cpp
@@ -0,0 +1,158 @@
+﻿// Copyright Epic Games, Inc. All Rights Reserved.
+
+#include "CoreMinimal.h"
+#include "TestHarness.h"
+
+#include <catch2/generators/catch_generators.hpp>
+
+#include "PCGContext.h"
+#include "Data/PCGSplineData.h"
+#include "Tests/PCGTestsCommon.h"
+#include "PCGTestsCommon.h"
+#include "Metadata/Accessors/PCGAttributeAccessorHelpers.h"
+
+
+TEST_CASE_METHOD(PCGTests::FPCGBaseTestWithBasicContext, "PCG::Metadata::Accessors::SplineAccessors", "[PCG][SplineAccessors]")
+{
+	const FName DoubleAttributeName = "Double";
+	
+	constexpr int32 NumPoints = 5;
+
+	TArray<FSplinePoint> SplinePoints;
+	for (int32 i = 0; i < NumPoints; ++i)
+	{
+		SplinePoints.Emplace_GetRef(i, FVector::OneVector * i * 100).Type = ESplinePointType::Linear;
+	}
+
+	UPCGSplineData* InputData = CreateData<UPCGSplineData>();
+	InputData->Initialize(SplinePoints, /*bClosedLoop=*/false, FTransform::Identity);
+
+
+	SECTION("Double attribute")
+	{
+		FPCGMetadataAttribute<double>* DoubleAttribute = InputData->MutableMetadata()->CreateAttribute<double>(DoubleAttributeName, 0.0, false, false);
+		
+		TUniquePtr<IPCGAttributeAccessor> DoubleAttributeAccessor = PCGAttributeAccessorHelpers::CreateAccessor(InputData, FPCGAttributePropertySelector::CreateAttributeSelector(DoubleAttributeName));
+		TUniquePtr<IPCGAttributeAccessorKeys> DoubleAttributeKeys = PCGAttributeAccessorHelpers::CreateKeys(InputData, FPCGAttributePropertySelector::CreateAttributeSelector(DoubleAttributeName));
+
+		CHECK(DoubleAttributeAccessor.IsValid());
+		CHECK(DoubleAttributeKeys.IsValid());
+
+		REQUIRE_EQUAL(InputData->MutableMetadata()->GetItemCountForChild(), 0);
+		REQUIRE_EQUAL(DoubleAttributeKeys->GetNum(), 5);
+
+		TStaticArray<double, NumPoints> Values{};
+		REQUIRE(DoubleAttributeAccessor->GetRange<double>(Values, 0, *DoubleAttributeKeys));
+
+		REQUIRE(Algo::AllOf(Values, [](const double Value) { return Value == 0.0; }));
+
+		Values = {0, 1, 2, 3, 4};
+		REQUIRE(DoubleAttributeAccessor->SetRange<double>(Values, 0, *DoubleAttributeKeys));
+		REQUIRE_EQUAL(InputData->MutableMetadata()->GetItemCountForChild(), 5);
+
+		TArray<PCGMetadataEntryKey> EntryKeys = InputData->GetMetadataEntryKeysForSplinePoints();
+		for (int32 i = 0; i < NumPoints; ++i)
+		{
+			REQUIRE_EQUAL(DoubleAttribute->GetValueFromItemKey(EntryKeys[i]), i);
+		}
+	}
+
+	SECTION("Double data attribute")
+	{
+		FPCGMetadataAttribute<double>* DoubleDataAttribute = InputData->MutableMetadata()->CreateAttribute<double>(FPCGAttributeIdentifier{DoubleAttributeName, PCGMetadataDomainID::Data}, 50.0, false, false);
+		
+		TUniquePtr<IPCGAttributeAccessor> DoubleDataAttributeAccessor = PCGAttributeAccessorHelpers::CreateAccessor(InputData, FPCGAttributePropertySelector::CreateAttributeSelector(DoubleAttributeName, TEXT("Data")));
+		TUniquePtr<IPCGAttributeAccessorKeys> DoubleDataAttributeKeys = PCGAttributeAccessorHelpers::CreateKeys(InputData, FPCGAttributePropertySelector::CreateAttributeSelector(DoubleAttributeName, TEXT("Data")));
+		
+		CHECK(DoubleDataAttributeAccessor.IsValid());
+		CHECK(DoubleDataAttributeKeys.IsValid());
+
+		REQUIRE_EQUAL(InputData->MutableMetadata()->GetMetadataDomain(PCGMetadataDomainID::Data)->GetItemCountForChild(), 0);
+		REQUIRE_EQUAL(DoubleDataAttributeKeys->GetNum(), 1);
+
+		double Value{};
+		REQUIRE(DoubleDataAttributeAccessor->Get<double>(Value, *DoubleDataAttributeKeys));
+
+		REQUIRE_EQUAL(Value, 50.0);
+		
+		REQUIRE(DoubleDataAttributeAccessor->Set<double>(150.0, *DoubleDataAttributeKeys));
+		REQUIRE_EQUAL(InputData->MutableMetadata()->GetMetadataDomain(PCGMetadataDomainID::Data)->GetItemCountForChild(), 1);
+
+		REQUIRE_EQUAL(DoubleDataAttribute->GetValueFromItemKey(PCGFirstEntryKey), 150.0);
+	}
+
+	SECTION("Length accessor")
+	{
+		// Length is read only, accessor cannot be created.
+		TUniquePtr<IPCGAttributeAccessor> LengthNonConstAccessor = PCGAttributeAccessorHelpers::CreateAccessor(InputData, FPCGAttributePropertySelector::CreatePropertySelector(TEXT("Length"), TEXT("Data")));
+		CHECK(!LengthNonConstAccessor.IsValid());
+		
+		TUniquePtr<const IPCGAttributeAccessor> LengthAccessor = PCGAttributeAccessorHelpers::CreateConstAccessor(InputData, FPCGAttributePropertySelector::CreatePropertySelector(TEXT("Length"), TEXT("Data")));
+		TUniquePtr<const IPCGAttributeAccessorKeys> LengthKeys = PCGAttributeAccessorHelpers::CreateConstKeys(InputData, FPCGAttributePropertySelector::CreatePropertySelector(TEXT("Length"), TEXT("Data")));
+		
+		CHECK(LengthAccessor.IsValid());
+		CHECK(LengthKeys.IsValid());
+
+		double Value{};
+		REQUIRE(LengthAccessor->Get(Value, *LengthKeys));
+		// Each point is located at i * 100, in X,Y,Z so each segment is sqrt(3) * 100. There are (NumPoints - 1) segments
+		static const double SegmentLength = FMath::Sqrt(3.0) * 100.0;
+		static const double ExpectedLength = SegmentLength * (NumPoints - 1);
+		REQUIRE(FMath::IsNearlyEqual(Value, ExpectedLength, /*ErrorTolerance=*/1e-4));
+	}
+	
+	SECTION("IsClosed accessor")
+    {
+    	// IsClosed is read only, accessor cannot be created.
+    	TUniquePtr<IPCGAttributeAccessor> IsClosedNonConstAccessor = PCGAttributeAccessorHelpers::CreateAccessor(InputData, FPCGAttributePropertySelector::CreatePropertySelector(TEXT("IsClosed"), TEXT("Data")));
+    	CHECK(!IsClosedNonConstAccessor.IsValid());
+    	
+    	TUniquePtr<const IPCGAttributeAccessor> IsClosedAccessor = PCGAttributeAccessorHelpers::CreateConstAccessor(InputData, FPCGAttributePropertySelector::CreatePropertySelector(TEXT("IsClosed"), TEXT("Data")));
+    	TUniquePtr<const IPCGAttributeAccessorKeys> IsClosedKeys = PCGAttributeAccessorHelpers::CreateConstKeys(InputData, FPCGAttributePropertySelector::CreatePropertySelector(TEXT("IsClosed"), TEXT("Data")));
+    	
+    	CHECK(IsClosedAccessor.IsValid());
+    	CHECK(IsClosedKeys.IsValid());
+
+    	bool Value = true;
+    	REQUIRE(IsClosedAccessor->Get(Value, *IsClosedKeys));
+    	REQUIRE_FALSE(Value);
+    }
+
+	SECTION("Arrive tangent accessor")
+	{
+		TUniquePtr<IPCGAttributeAccessor> ArriveTangentAccessor = PCGAttributeAccessorHelpers::CreateAccessor(InputData, FPCGAttributePropertySelector::CreatePropertySelector(TEXT("ArriveTangent")));
+		TUniquePtr<IPCGAttributeAccessorKeys> ArriveTangentKeys = PCGAttributeAccessorHelpers::CreateKeys(InputData, FPCGAttributePropertySelector::CreatePropertySelector(TEXT("ArriveTangent")));
+		
+		CHECK(ArriveTangentAccessor.IsValid());
+		CHECK(ArriveTangentKeys.IsValid());
+
+		TStaticArray<FVector, NumPoints> Values{};
+		REQUIRE(ArriveTangentAccessor->GetRange<FVector>(Values, 0, *ArriveTangentKeys));
+
+		for (int32 i = 0; i < NumPoints; ++i)
+		{
+			if (i == 0)
+			{
+				// No arrive tangent on first point
+				REQUIRE_EQUAL(Values[0], FVector::ZeroVector);
+			}
+			else
+			{
+				REQUIRE_EQUAL(Values[i], FVector::OneVector * 100);
+			}
+		}
+
+		for (int32 i = 0; i < NumPoints; ++i)
+		{
+			Values[i] = FVector::OneVector * (i + 1);
+		}
+
+		REQUIRE(ArriveTangentAccessor->SetRange<FVector>(Values, 0, *ArriveTangentKeys));
+
+		TArray<FSplinePoint> Points = InputData->GetSplinePoints();
+		for (int32 i = 0; i < NumPoints - 1; ++i)
+		{
+			REQUIRE_EQUAL(Points[i].ArriveTangent , FVector::OneVector * (i + 1));
+		}
+	}
+}
