[PCG] Fix tests that were logging an error: We can now suppres error logging if we test for errors

#rnx

#Synced-CL 49805392

diff --git a/Engine/Plugins/PCG/Tests/Private/Elements/PCGMetadataVectorOpElementTests.cpp b/Engine/Plugins/PCG/Tests/Private/Elements/PCGMetadataVectorOpElementTests.cpp
--- a/Engine/Plugins/PCG/Tests/Private/Elements/PCGMetadataVectorOpElementTests.cpp
+++ b/Engine/Plugins/PCG/Tests/Private/Elements/PCGMetadataVectorOpElementTests.cpp
@@ -285,6 +285,7 @@ TEST_CASE_METHOD(PCGTests::FPCGVectorOpsBaseTest, "PCG::VectorOps::RotatorBetwee
 		TArray<FVector4> InputB = GenerateValues<FVector4>(FVector4{1, 1, 0, 1}, FVector4{0, 1, 0, 1}, 5);
 
 		
+		FSuppressErrorsScope SuppressErrors(*this);
 		RunTest(TestParams(EPCGMetadataVectorOperation::RotatorBetweenVectors, GenerateParamData<FVector4>({InputA,InputB})));
 		CHECK(NumErrors > 0);
 	}
@@ -357,6 +358,7 @@ TEST_CASE_METHOD(PCGTests::FPCGVectorOpsBaseTest, "PCG::VectorOps::EulerBetweenV
 		TArray<FVector4> InputB = GenerateValues<FVector4>(FVector4{1, 1, 0, 1}, FVector4{0, 1, 0, 1}, 5);
 
 		
+		FSuppressErrorsScope SuppressErrors(*this);
 		RunTest(TestParams(EPCGMetadataVectorOperation::EulerBetweenVectors, GenerateParamData<FVector4>({InputA,InputB})));
 		CHECK(NumErrors > 0);
 	}
@@ -407,7 +409,8 @@ TEST_CASE_METHOD(PCGTests::FPCGVectorOpsBaseTest, "PCG::VectorOps::AngleBetween2
 		{
 			TArray<FVector> InputA = GenerateValues<FVector>(FVector{0, 1, 0}, FVector{1, 0, 1}, 5);
 			TArray<FVector> InputB = GenerateValues<FVector>(FVector{1, 1, 0}, FVector{1, 0, 1}, 5);
-		
+			
+			FSuppressErrorsScope SuppressErrors(*this);
 			RunTest(TestParams(EPCGMetadataVectorOperation::AngleBetweenVector2D, GenerateParamData<FVector>({InputA,InputB})));
 			CHECK(NumErrors > 0);
 		}
@@ -416,6 +419,7 @@ TEST_CASE_METHOD(PCGTests::FPCGVectorOpsBaseTest, "PCG::VectorOps::AngleBetween2
 			TArray<FVector4> InputA = GenerateValues<FVector4>(FVector4{0, 1, 0, 1}, FVector4{1, 1, 0, 1}, 5);
 			TArray<FVector4> InputB = GenerateValues<FVector4>(FVector4{1, 1, 0, 1}, FVector4{0, 1, 0, 1}, 5);
 		
+			FSuppressErrorsScope SuppressErrors(*this);
 			RunTest(TestParams(EPCGMetadataVectorOperation::AngleBetweenVector2D, GenerateParamData<FVector4>({InputA,InputB})));
 			CHECK(NumErrors > 0);
 		}
diff --git a/Engine/Plugins/PCG/Tests/Private/PCGTestsCommon.cpp b/Engine/Plugins/PCG/Tests/Private/PCGTestsCommon.cpp
--- a/Engine/Plugins/PCG/Tests/Private/PCGTestsCommon.cpp
+++ b/Engine/Plugins/PCG/Tests/Private/PCGTestsCommon.cpp
@@ -22,14 +22,18 @@ namespace PCGTests
 		return FPCGContext::NewPointData_AnyThread(GetContext());
 	}
 
-	FPCGTestsLogOutputDevice::FPCGTestsLogOutputDevice(): ThreadID(FPlatformTLS::GetCurrentThreadId())
+	FPCGTestsLogOutputDevice::FPCGTestsLogOutputDevice(bool bInSuppressErrors):
+		bSuppressErrors(bInSuppressErrors),
+		OldContext(GWarn)
 	{
+		GWarn = this;
 		GLog->AddOutputDevice(this);
 	}
 
 	FPCGTestsLogOutputDevice::~FPCGTestsLogOutputDevice()
 	{
 		GLog->RemoveOutputDevice(this);
+		GWarn = OldContext;
 	}
 
 	void FPCGTestsLogOutputDevice::Serialize(const TCHAR* V, ELogVerbosity::Type Verbosity, const class FName& Category)
@@ -44,6 +48,14 @@ namespace PCGTests
 			{
 				NbMessageReceived.Add(Verbosity, 1);
 			}
+			if (!bSuppressErrors)
+			{
+				OldContext->Serialize(V, Verbosity, Category);
+			}
+		}
+		else
+		{
+			OldContext->Serialize(V, Verbosity, Category);
 		}
 	}
 }
diff --git a/Engine/Plugins/PCG/Tests/Private/PCGTestsCommon.h b/Engine/Plugins/PCG/Tests/Private/PCGTestsCommon.h
--- a/Engine/Plugins/PCG/Tests/Private/PCGTestsCommon.h
+++ b/Engine/Plugins/PCG/Tests/Private/PCGTestsCommon.h
@@ -8,6 +8,8 @@
 #include <catch2/catch_approx.hpp>
 #include <catch2/matchers/catch_matchers_all.hpp>
 
+#include "Misc/FeedbackContext.h"
+
 class FPCGModule;
 class UPCGBasePointData;
 
@@ -109,9 +111,9 @@ protected:
 };
 	
 /** Utility class to gather errors and warnings*/
-struct FPCGTestsLogOutputDevice : public FOutputDevice
+struct FPCGTestsLogOutputDevice : public FFeedbackContext
 {
-	FPCGTestsLogOutputDevice();
+	explicit FPCGTestsLogOutputDevice(bool bInSuppressErrors);
 
 	virtual ~FPCGTestsLogOutputDevice();
 
@@ -121,9 +123,10 @@ struct FPCGTestsLogOutputDevice : public FOutputDevice
 	virtual bool CanBeUsedOnAnyThread() const override { return true; }
 	virtual void Serialize(const TCHAR* V, ELogVerbosity::Type Verbosity, const class FName& Category) override;
 
+	bool bSuppressErrors;
+	FFeedbackContext* OldContext = nullptr;
 	using MessageMap = TMap<ELogVerbosity::Type, int32>;
 	MessageMap NbMessageReceived;
-	const uint32 ThreadID;
 };
 	
 /** Base templated class to execute a PCG element.*/
@@ -146,6 +149,22 @@ public:
 	}
 
 	virtual FPCGContext* GetContext() override {return Context.Get();};
+	
+	class FSuppressErrorsScope
+	{
+		public:
+		explicit FSuppressErrorsScope(FPCGSingleElementBaseTest& InTestReference): TestReference(InTestReference)
+		{
+			TestReference.bSuppressErrors = true;
+		}
+		~FSuppressErrorsScope()
+		{
+			TestReference.bSuppressErrors = false;
+		}
+		
+	private:
+		FPCGSingleElementBaseTest& TestReference;
+	};
 
 	void ExecuteElement()
 	{
@@ -154,7 +173,7 @@ public:
 		Context->InitializeSettings();
 		Context->AsyncState.NumAvailableTasks = 1;
 		
-		PCGTests::FPCGTestsLogOutputDevice LogCapture;
+		PCGTests::FPCGTestsLogOutputDevice LogCapture(bSuppressErrors);
 
 		while (!TestElement->Execute(Context.Get()))
 		{
@@ -168,6 +187,7 @@ public:
 	FPCGElementPtr TestElement;
 	FPCGDataCollection InputData;
 	TUniquePtr<FPCGContext> Context;
+	bool bSuppressErrors = false;
 	int NumErrors = 0;
 	int NumWarnings = 0;
 };	
