[PCG] Fixing unit tests for Clang

#rnx

#Synced-CL 50405705

diff --git a/Engine/Plugins/PCG/Tests/Private/Metadata/PCGMetadataFlattenTests.cpp b/Engine/Plugins/PCG/Tests/Private/Metadata/PCGMetadataFlattenTests.cpp
--- a/Engine/Plugins/PCG/Tests/Private/Metadata/PCGMetadataFlattenTests.cpp
+++ b/Engine/Plugins/PCG/Tests/Private/Metadata/PCGMetadataFlattenTests.cpp
@@ -163,7 +163,7 @@ TEST_CASE_METHOD(PCGTests::FPCGBaseTestWithBasicContext, "PCG::Metadata::Flatten
 
 	// Validate the values
 	const TConstPCGValueRange<int64> FinalMetadataEntryRange = FinalPointData->GetConstMetadataEntryValueRange();
-	auto i = GENERATE(range(0, NumPoints));
+	auto i = GENERATE_COPY(range(0, NumPoints));
 	
 	SECTION("Validate values")
 	{
