[PCG] Add unit test for accessor with Generic Attributes and multi-thread writing

* Also update the metadata domain for default values on the interpolation and override parent

#rnx

#Synced-CL 50489647

diff --git a/Engine/Plugins/PCG/Source/PCG/Public/Metadata/PCGMetadataDomain.h b/Engine/Plugins/PCG/Source/PCG/Public/Metadata/PCGMetadataDomain.h
--- a/Engine/Plugins/PCG/Source/PCG/Public/Metadata/PCGMetadataDomain.h
+++ b/Engine/Plugins/PCG/Source/PCG/Public/Metadata/PCGMetadataDomain.h
@@ -236,9 +236,12 @@ public:
 	* Create a new generic attribute.
 	*/
 	template<typename T>
-	FPCGMetadataAttributeGeneric* CreateGenericAttribute(FName AttributeName, const T& DefaultValue, bool bAllowsInterpolation, bool bOverrideParent);
+	FPCGMetadataAttributeGeneric* CreateGenericAttribute(FName AttributeName, const T& DefaultValue, bool bAllowsInterpolation = true, bool bOverrideParent = false);
 	
-	UE_API FPCGMetadataAttributeGeneric* CreateGenericAttribute(const FPCGMetadataAttributeDesc& Desc, bool bAllowsInterpolation, bool bOverrideParent);
+	/**
+	* Create a new generic attribute from a desc.
+	*/
+	UE_API FPCGMetadataAttributeGeneric* CreateGenericAttribute(const FPCGMetadataAttributeDesc& Desc, bool bAllowsInterpolation = true, bool bOverrideParent = false);
 
 	UE_API FPCGMetadataAttributeGeneric* GetMutableGenericAttribute(FName AttributeName);
 	UE_API const FPCGMetadataAttributeGeneric* GetConstGenericAttribute(FName AttributeName) const;
diff --git a/Engine/Plugins/PCG/Tests/Private/Benchmarks/Metadata/PCGMetadataAttributeGenericBenchmarks.cpp b/Engine/Plugins/PCG/Tests/Private/Benchmarks/Metadata/PCGMetadataAttributeGenericBenchmarks.cpp
--- a/Engine/Plugins/PCG/Tests/Private/Benchmarks/Metadata/PCGMetadataAttributeGenericBenchmarks.cpp
+++ b/Engine/Plugins/PCG/Tests/Private/Benchmarks/Metadata/PCGMetadataAttributeGenericBenchmarks.cpp
@@ -35,8 +35,8 @@ TEST_CASE_METHOD(PCGTests::FPCGBaseTestWithBasicContext, "PCG::AttributeGeneric:
 		TPCGValueRange<int64> PointArrayDataMetadataEntryRange = PointArrayData->GetMetadataEntryValueRange(/*bAllocate=*/false);
 
 		FPCGMetadataDomain* Domain = PointArrayData->MutableMetadata()->GetDefaultMetadataDomain();
-		FPCGMetadataAttributeGeneric* GenericAttribute = Domain->CreateGenericAttribute<double>(DoubleGenericAttributeName, 0., true, false);
-		FPCGMetadataAttribute<double>* TypedAttribute = Domain->CreateAttribute<double>(DoubleAttributeName, 0., true, false);
+		FPCGMetadataAttributeGeneric* GenericAttribute = Domain->CreateGenericAttribute<double>(DoubleGenericAttributeName, 0.);
+		FPCGMetadataAttribute<double>* TypedAttribute = Domain->CreateAttribute<double>(DoubleAttributeName, 0., /*bAllowsInterpolation=*/true, /*bOverrideParent=*/false);
 		
 		// Allocate entry keys and values
 		TStaticArray<double, NumPoints> Values;
@@ -162,8 +162,8 @@ TEST_CASE_METHOD(PCGTests::FPCGBaseTestWithBasicContext, "PCG::AttributeGeneric:
 		TPCGValueRange<int64> PointArrayDataMetadataEntryRange = PointArrayData->GetMetadataEntryValueRange(/*bAllocate=*/false);
 
 		FPCGMetadataDomain* Domain = PointArrayData->MutableMetadata()->GetDefaultMetadataDomain();
-		FPCGMetadataAttributeGeneric* GenericAttribute = Domain->CreateGenericAttribute<FString>(StringGenericAttributeName, FString{}, true, false);
-		FPCGMetadataAttribute<FString>* TypedAttribute = Domain->CreateAttribute<FString>(StringAttributeName, FString{}, true, false);
+		FPCGMetadataAttributeGeneric* GenericAttribute = Domain->CreateGenericAttribute<FString>(StringGenericAttributeName, FString{});
+		FPCGMetadataAttribute<FString>* TypedAttribute = Domain->CreateAttribute<FString>(StringAttributeName, FString{}, /*bAllowsInterpolation=*/true, /*bOverrideParent=*/false);
 
 		// Allocate entry keys and values
 		TStaticArray<FString, NumPoints> Values;
diff --git a/Engine/Plugins/PCG/Tests/Private/Metadata/Accessors/PCGMetadataAttributeGenericAccessorTests.cpp b/Engine/Plugins/PCG/Tests/Private/Metadata/Accessors/PCGMetadataAttributeGenericAccessorTests.cpp
new file mode 100644
--- /dev/null
+++ b/Engine/Plugins/PCG/Tests/Private/Metadata/Accessors/PCGMetadataAttributeGenericAccessorTests.cpp
@@ -0,0 +1,62 @@
+﻿// Copyright Epic Games, Inc. All Rights Reserved.
+
+#include "CoreMinimal.h"
+#include "TestHarness.h"
+
+#include <catch2/generators/catch_generators.hpp>
+
+#include "PCGContext.h"
+#include "Data/PCGPointArrayData.h"
+#include "Tests/PCGTestsCommon.h"
+#include "PCGTestsCommon.h"
+#include "Metadata/Accessors/PCGAttributeAccessorHelpers.h"
+
+
+TEST_CASE_METHOD(PCGTests::FPCGBaseTestWithBasicContext, "PCG::Metadata::Accessors::GenericAttribute", "[PCG][Accessors][GenericAttribute]")
+{
+	const FName DoubleAttributeName = "Double";
+	
+	// Preparing the data
+	constexpr int32 NumPoints = 5;
+	constexpr int32 Seed = 42;
+	UPCGPointArrayData* PointArrayData = CreateData<UPCGPointArrayData>();
+
+	PointArrayData->SetNumPoints(NumPoints);
+
+	PointArrayData->AllocateProperties(EPCGPointNativeProperties::MetadataEntry);
+	TPCGValueRange<int64> PointArrayDataMetadataEntryRange = PointArrayData->GetMetadataEntryValueRange(/*bAllocate=*/false);
+
+	FPCGMetadataDomain* Domain = PointArrayData->MutableMetadata()->GetDefaultMetadataDomain();
+		
+	FRandomStream RandomStream(Seed);
+		
+	FPCGMetadataAttributeGeneric* GenericAttribute = Domain->CreateGenericAttribute<double>(DoubleAttributeName, -1.0);
+	REQUIRE_NOT_EQUAL(GenericAttribute, nullptr);
+	
+	TUniquePtr<IPCGAttributeAccessor> DoubleAttributeAccessor = PCGAttributeAccessorHelpers::CreateAccessor(PointArrayData, FPCGAttributePropertySelector::CreateAttributeSelector(DoubleAttributeName));
+	TUniquePtr<IPCGAttributeAccessorKeys> DoubleAttributeKeys = PCGAttributeAccessorHelpers::CreateKeys(PointArrayData, FPCGAttributePropertySelector::CreateAttributeSelector(DoubleAttributeName));
+
+	REQUIRE(DoubleAttributeAccessor.IsValid());
+	REQUIRE(DoubleAttributeKeys.IsValid());
+	
+	REQUIRE_EQUAL(PointArrayData->MutableMetadata()->GetItemCountForChild(), 0);
+	REQUIRE_EQUAL(DoubleAttributeKeys->GetNum(), 5);
+
+	TStaticArray<double, NumPoints> Values{};
+	REQUIRE(DoubleAttributeAccessor->GetRange<double>(Values, 0, *DoubleAttributeKeys));
+
+	REQUIRE(Algo::AllOf(Values, [](const double Value) { return Value == -1.0; }));
+
+	Values = {0, 1, 2, 3, 4};
+	REQUIRE(DoubleAttributeAccessor->SetRange<double>(Values, 0, *DoubleAttributeKeys));
+	REQUIRE_EQUAL(PointArrayData->MutableMetadata()->GetItemCountForChild(), 5);
+	
+	for (int32 i = 0; i < NumPoints; ++i)
+	{
+		REQUIRE_EQUAL(GenericAttribute->GetValueFromItemKey<double>(PointArrayDataMetadataEntryRange[i]), i);
+	}
+	
+	TStaticArray<double, NumPoints> ReadValues{};
+	REQUIRE(DoubleAttributeAccessor->GetRange<double>(ReadValues, 0, *DoubleAttributeKeys));
+	CHECK_THAT(Values, Catch::Matchers::RangeEquals(ReadValues, [](const auto& a, const auto& b){ return a == b; } ));
+}
\ No newline at end of file
diff --git a/Engine/Plugins/PCG/Tests/Private/Metadata/Accessors/PCGSplineAccessorTests.cpp b/Engine/Plugins/PCG/Tests/Private/Metadata/Accessors/PCGSplineAccessorTests.cpp
--- a/Engine/Plugins/PCG/Tests/Private/Metadata/Accessors/PCGSplineAccessorTests.cpp
+++ b/Engine/Plugins/PCG/Tests/Private/Metadata/Accessors/PCGSplineAccessorTests.cpp
@@ -35,8 +35,8 @@ TEST_CASE_METHOD(PCGTests::FPCGBaseTestWithBasicContext, "PCG::Metadata::Accesso
 		TUniquePtr<IPCGAttributeAccessor> DoubleAttributeAccessor = PCGAttributeAccessorHelpers::CreateAccessor(InputData, FPCGAttributePropertySelector::CreateAttributeSelector(DoubleAttributeName));
 		TUniquePtr<IPCGAttributeAccessorKeys> DoubleAttributeKeys = PCGAttributeAccessorHelpers::CreateKeys(InputData, FPCGAttributePropertySelector::CreateAttributeSelector(DoubleAttributeName));
 
-		CHECK(DoubleAttributeAccessor.IsValid());
-		CHECK(DoubleAttributeKeys.IsValid());
+		REQUIRE(DoubleAttributeAccessor.IsValid());
+		REQUIRE(DoubleAttributeKeys.IsValid());
 
 		REQUIRE_EQUAL(InputData->MutableMetadata()->GetItemCountForChild(), 0);
 		REQUIRE_EQUAL(DoubleAttributeKeys->GetNum(), 5);
diff --git a/Engine/Plugins/PCG/Tests/Private/Metadata/PCGMetadataAttributeGenericTests.cpp b/Engine/Plugins/PCG/Tests/Private/Metadata/PCGMetadataAttributeGenericTests.cpp
--- a/Engine/Plugins/PCG/Tests/Private/Metadata/PCGMetadataAttributeGenericTests.cpp
+++ b/Engine/Plugins/PCG/Tests/Private/Metadata/PCGMetadataAttributeGenericTests.cpp
@@ -13,6 +13,8 @@
 #include <catch2/generators/catch_generators_range.hpp>
 
 #include "TestHarness.h"
+#include "Async/Async.h"
+#include "HAL/Event.h"
 
 namespace PCGAttributeGenericTests
 {
@@ -396,13 +398,13 @@ TEST_CASE_METHOD(PCGTests::FPCGBaseTestWithBasicContext, "PCG::AttributeGeneric:
 	
 		if (AttributeTester.OverrideDesc.IsSet())
 		{
-			GenericAttribute = Domain->CreateGenericAttribute(*AttributeTester.OverrideDesc, false, false);
+			GenericAttribute = Domain->CreateGenericAttribute(*AttributeTester.OverrideDesc);
 			REQUIRE_NOT_EQUAL(GenericAttribute, nullptr);
 			GenericAttribute->SetDefaultValue<T>(AttributeTester.DefaultValue);
 		}
 		else
 		{
-			GenericAttribute = Domain->CreateGenericAttribute<T>(PCGAttributeGenericTests::AttributeName, AttributeTester.DefaultValue, false, false);
+			GenericAttribute = Domain->CreateGenericAttribute<T>(PCGAttributeGenericTests::AttributeName, AttributeTester.DefaultValue);
 			REQUIRE_NOT_EQUAL(GenericAttribute, nullptr);
 		}
 		
@@ -566,13 +568,13 @@ TEST_CASE_METHOD(PCGTests::FPCGBaseTestWithBasicContext, "PCG::AttributeGeneric:
 	
 		if (AttributeTester.OverrideDesc.IsSet())
 		{
-			GenericAttribute = Domain->CreateGenericAttribute(*AttributeTester.OverrideDesc, false, false);
+			GenericAttribute = Domain->CreateGenericAttribute(*AttributeTester.OverrideDesc);
 			REQUIRE_NOT_EQUAL(GenericAttribute, nullptr);
 			GenericAttribute->SetDefaultValue<T>(AttributeTester.DefaultValue);
 		}
 		else
 		{
-			GenericAttribute = Domain->CreateGenericAttribute<T>(PCGAttributeGenericTests::AttributeName, AttributeTester.DefaultValue, false, false);
+			GenericAttribute = Domain->CreateGenericAttribute<T>(PCGAttributeGenericTests::AttributeName, AttributeTester.DefaultValue);
 			REQUIRE_NOT_EQUAL(GenericAttribute, nullptr);
 		}
 		
@@ -713,7 +715,7 @@ TEST_CASE_METHOD(PCGTests::FPCGBaseTestWithBasicContext, "PCG::AttributeGeneric:
 		UPCGParamData* ParamData = CreateData<UPCGParamData>();
 		FPCGMetadataDomain* Domain = ParamData->MutableMetadata()->GetDefaultMetadataDomain();
 		
-		FPCGMetadataAttributeGeneric* GenericAttribute = Domain->CreateGenericAttribute<RealType>(PCGAttributeGenericTests::AttributeName, InValue, false, false);
+		FPCGMetadataAttributeGeneric* GenericAttribute = Domain->CreateGenericAttribute<RealType>(PCGAttributeGenericTests::AttributeName, InValue);
 		
 		SECTION("Get Value with wrong type")
 		{
@@ -847,7 +849,7 @@ TEST_CASE_METHOD(PCGTests::FPCGBaseTestWithBasicContext, "PCG::AttributeGeneric:
 	FPCGMetadataDomain* Domain = PointArrayData->MutableMetadata()->GetDefaultMetadataDomain();
 	
 	// Initialize the default value to a value we will add later on.
-	FPCGMetadataAttributeGeneric* GenericAttribute = Domain->CreateGenericAttribute<FString>(PCGAttributeGenericTests::AttributeName, TEXT("Value: 0"), false, false);
+	FPCGMetadataAttributeGeneric* GenericAttribute = Domain->CreateGenericAttribute<FString>(PCGAttributeGenericTests::AttributeName, TEXT("Value: 0"));
 	REQUIRE_NOT_EQUAL(GenericAttribute, nullptr);
 	
 	// First set half the points with different values
@@ -894,4 +896,242 @@ TEST_CASE_METHOD(PCGTests::FPCGBaseTestWithBasicContext, "PCG::AttributeGeneric:
 	GenericAttribute->GetValuesFromItemKeys<const FString*>(PCGValueRangeHelpers::MakeConstValueRange(PointArrayDataMetadataEntryRange), ReadValues);
 	
 	CHECK_THAT(Values, Catch::Matchers::RangeEquals(ReadValues, [](const FString& a, const FString* b){ return a == *b; } ));
+}
+
+TEST_CASE_METHOD(PCGTests::FPCGBaseTestWithBasicContext, "PCG::AttributeGeneric::MultiThreadedWrite", "[PCG][AttributeGeneric]")
+{
+	auto Tester = [this]<typename T, typename ReadType>(PCGAttributeGenericTests::TypedAttributeTester<T, ReadType> AttributeTester, bool bLockless)
+	{
+		// Preparing the data
+		constexpr int32 NumEntries = 10000;
+		constexpr int32 Chunks = 100;
+		constexpr int32 ChunksPerThread = 10;
+		constexpr int32 ElementsPerThreads = Chunks * ChunksPerThread;
+		constexpr int32 NumThreads = NumEntries / ElementsPerThreads;
+		
+		constexpr int32 Seed = 42;
+		UPCGParamData* ParamData = CreateData<UPCGParamData>();
+		FPCGMetadataDomain* Domain = ParamData->MutableMetadata()->GetDefaultMetadataDomain();
+	
+		FRandomStream RandomStream(Seed);
+	
+		FPCGMetadataAttributeGeneric* GenericAttribute = Domain->CreateGenericAttribute<T>(PCGAttributeGenericTests::AttributeName, AttributeTester.DefaultValue);
+		REQUIRE_NOT_EQUAL(GenericAttribute, nullptr);
+	
+		TStaticArray<T, NumEntries> Values{};
+		TStaticArray<PCGMetadataEntryKey, NumEntries> EntryKeys;
+		TStaticArray<PCGMetadataEntryKey*, NumEntries> EntryKeysPtr;
+		int32 StartIndex = INDEX_NONE;
+		
+		for (int32 i = 0; i < NumEntries; ++i)
+		{
+			EntryKeys[i] = Domain->AddEntry();
+			EntryKeysPtr[i] = &EntryKeys[i];
+			AttributeTester.GenerateRandom(RandomStream, Values[i]);
+		}
+		
+		if (bLockless)
+		{
+			StartIndex = GenericAttribute->PreallocateValues(EntryKeysPtr, /*bLockless=*/true);
+		}
+		
+		FEvent* StartEvent = FPlatformProcess::GetSynchEventFromPool(/*bIsManualReset=*/true);
+		FEvent* LastEvent = FPlatformProcess::GetSynchEventFromPool(/*bIsManualReset=*/true);
+		TAtomic<uint64> ActiveCount{ 0 };
+		TAtomic<bool> bSuccess = true;
+		TAtomic<uint64> NumActualThreads = 0;
+	
+		for (int32 i = 0; i < NumThreads; ++i)
+		{
+			auto Process = [&Values, &EntryKeysPtr, StartEvent, i, ElementsPerThreads, Chunks, ChunksPerThread, bLockless, &bSuccess, GenericAttribute, StartIndex, &ActiveCount, LastEvent, &NumActualThreads]()
+			{
+				if (StartEvent->Wait(2000))
+				{
+					for (int32 j = 0; j < ChunksPerThread; ++j)
+					{
+						const int32 ValuesIndex = i * ElementsPerThreads + j * Chunks;
+						TArrayView<PCGMetadataEntryKey*> EntryKeysView = MakeArrayView(EntryKeysPtr.GetData() + ValuesIndex, Chunks);
+						TArrayView<T> ValuesView = MakeArrayView(Values.GetData() + ValuesIndex, Chunks);
+					
+						if (bLockless)
+						{
+							GenericAttribute->SetValues_TryLockless<T>(EntryKeysView, ValuesView, StartIndex + ValuesIndex);
+						}
+						else
+						{
+							GenericAttribute->SetValues<T>(EntryKeysView, ValuesView);
+						}
+					}
+				}
+				else
+				{
+					bSuccess = false;
+				}
+					
+				// Verify that it was actually done in MT. 
+				if (!IsInGameThread())
+				{
+					++NumActualThreads;
+				}
+					
+				if (--ActiveCount == 0)
+				{
+					LastEvent->Trigger();
+				}
+			};
+			
+			++ActiveCount;
+			Async(i < FTaskGraphInterface::Get().GetNumWorkerThreads() ? EAsyncExecution::TaskGraph : EAsyncExecution::Thread, MoveTemp(Process));
+		}
+		
+		StartEvent->Trigger();
+		LastEvent->Wait();
+		
+		FPlatformProcess::ReturnSynchEventToPool(StartEvent);
+		FPlatformProcess::ReturnSynchEventToPool(LastEvent);
+		
+		REQUIRE(bSuccess);
+		
+		// Making sure we had concurrency.
+		CHECK(NumActualThreads > 1);
+		
+		REQUIRE_EQUAL(GenericAttribute->GetNumberOfEntries(), Values.Num());
+
+		SECTION("Read back and verify values")
+		{
+			TStaticArray<ReadType, NumEntries> ReadValues;
+			GenericAttribute->GetValuesFromItemKeys<ReadType>(EntryKeys, ReadValues);
+			CHECK_THAT(Values, Catch::Matchers::RangeEquals(ReadValues, [&AttributeTester](const auto& a, const auto& b){ return AttributeTester.Verify(a, b); } ));
+		}
+	};
+
+	SECTION("Float Attribute")
+	{
+		SECTION("With lock")
+		{
+			Tester(PCGAttributeGenericTests::FloatTester(), /*bLockless=*/false);
+		}
+		
+		SECTION("Without lock")
+		{
+			Tester(PCGAttributeGenericTests::FloatTester(), /*bLockless=*/true);
+		}
+	}
+	
+	SECTION("Double Attribute")
+	{
+		SECTION("With lock")
+		{
+			Tester(PCGAttributeGenericTests::DoubleTester(), /*bLockless=*/false);
+		}
+		
+		SECTION("Without lock")
+		{
+			Tester(PCGAttributeGenericTests::DoubleTester(), /*bLockless=*/true);
+		}
+	}
+	
+	SECTION("Vector Attribute")
+	{
+		SECTION("With lock")
+		{
+			Tester(PCGAttributeGenericTests::VectorTester_Basic(), /*bLockless=*/false);
+		}
+		
+		SECTION("Without lock")
+		{
+			Tester(PCGAttributeGenericTests::VectorTester_Basic(), /*bLockless=*/true);
+		}
+	}
+	
+	SECTION("String Attribute")
+	{
+		SECTION("With lock")
+		{
+			Tester(PCGAttributeGenericTests::StringTester(), /*bLockless=*/false);
+		}
+		
+		SECTION("Without lock")
+		{
+			Tester(PCGAttributeGenericTests::StringTester(), /*bLockless=*/true);
+		}
+	}
+	
+	SECTION("FPCGPoint Attribute")
+	{		
+		SECTION("With lock")
+		{
+			Tester(PCGAttributeGenericTests::FPCGPointTester(), /*bLockless=*/false);
+		}
+		
+		SECTION("Without lock")
+		{
+			Tester(PCGAttributeGenericTests::FPCGPointTester(), /*bLockless=*/true);
+		}
+	}
+	
+	SECTION("Array of Double Attribute")
+	{		
+		SECTION("With lock")
+		{
+			Tester(PCGAttributeGenericTests::ArrayDoubleTester(), /*bLockless=*/false);
+		}
+		
+		SECTION("Without lock")
+		{
+			Tester(PCGAttributeGenericTests::ArrayDoubleTester(), /*bLockless=*/true);
+		}
+	}
+	
+	SECTION("Array of FPCGPoint Attribute")
+	{	
+		SECTION("With lock")
+		{
+			Tester(PCGAttributeGenericTests::ArrayFPCGPointTester(), /*bLockless=*/false);
+		}
+		
+		SECTION("Without lock")
+		{
+			Tester(PCGAttributeGenericTests::ArrayFPCGPointTester(), /*bLockless=*/true);
+		}
+	}
+	
+	SECTION("Set of Double Attribute")
+	{		
+		SECTION("With lock")
+		{
+			Tester(PCGAttributeGenericTests::SetDoubleTester(), /*bLockless=*/false);
+		}
+		
+		SECTION("Without lock")
+		{
+			Tester(PCGAttributeGenericTests::SetDoubleTester(), /*bLockless=*/true);
+		}
+	}
+	
+	SECTION("Map of String -> Double Attribute")
+	{
+		SECTION("With lock")
+		{
+			Tester(PCGAttributeGenericTests::MapStringDoubleTester(), /*bLockless=*/false);
+		}
+		
+		SECTION("Without lock")
+		{
+			Tester(PCGAttributeGenericTests::MapStringDoubleTester(), /*bLockless=*/true);
+		}
+	}
+	
+	SECTION("Enum attribute")
+	{
+		SECTION("With lock")
+		{
+			Tester(PCGAttributeGenericTests::EnumTester(), /*bLockless=*/false);
+		}
+		
+		SECTION("Without lock")
+		{
+			Tester(PCGAttributeGenericTests::EnumTester(), /*bLockless=*/true);
+		}
+	}
 }
\ No newline at end of file
diff --git a/Engine/Plugins/PCG/Tests/Private/Metadata/PCGMetadataContainerTypesTests.cpp b/Engine/Plugins/PCG/Tests/Private/Metadata/PCGMetadataContainerTypesTests.cpp
--- a/Engine/Plugins/PCG/Tests/Private/Metadata/PCGMetadataContainerTypesTests.cpp
+++ b/Engine/Plugins/PCG/Tests/Private/Metadata/PCGMetadataContainerTypesTests.cpp
@@ -29,7 +29,7 @@ TEST_CASE_METHOD(PCGTests::FPCGBaseTestWithBasicContext, "PCG::MetadataContainer
 	
 	SECTION("TScriptSetWrapper")
 	{
-		FPCGMetadataAttributeGeneric* SetGenericAttribute = Domain->CreateGenericAttribute<TSet<double>>("SetAttr", TSet<double>{}, false, false);
+		FPCGMetadataAttributeGeneric* SetGenericAttribute = Domain->CreateGenericAttribute<TSet<double>>("SetAttr", TSet<double>{});
 		REQUIRE_NOT_EQUAL(SetGenericAttribute, nullptr);
 		
 		const TSet<double> Values = {1.0, 2.0, 5.0, 3.0};
@@ -103,7 +103,7 @@ TEST_CASE_METHOD(PCGTests::FPCGBaseTestWithBasicContext, "PCG::MetadataContainer
 	
 	SECTION("TScriptMapWrapper")
 	{
-		FPCGMetadataAttributeGeneric* MapGenericAttribute = Domain->CreateGenericAttribute<TMap<FString, double>>("MapAttr", TMap<FString, double>{}, false, false);
+		FPCGMetadataAttributeGeneric* MapGenericAttribute = Domain->CreateGenericAttribute<TMap<FString, double>>("MapAttr", TMap<FString, double>{});
 		REQUIRE_NOT_EQUAL(MapGenericAttribute, nullptr);
 		
 		TMap<FString, double> Values = 
