[PCG] Implement missing functions for Generic Attributes

* Added SetValue/GetValueFromItemKey for parity with the typed attributes
* Implement Flatten/FlattenAndCompress/Prepare
* Added tests for flatten for generic attributes

#rnx

#Synced-CL 50489647

diff --git a/Engine/Plugins/PCG/Source/PCG/Private/Metadata/PCGMetadataAttributeGeneric.cpp b/Engine/Plugins/PCG/Source/PCG/Private/Metadata/PCGMetadataAttributeGeneric.cpp
--- a/Engine/Plugins/PCG/Source/PCG/Private/Metadata/PCGMetadataAttributeGeneric.cpp
+++ b/Engine/Plugins/PCG/Source/PCG/Private/Metadata/PCGMetadataAttributeGeneric.cpp
@@ -200,13 +200,13 @@ namespace PCG::Private
 			ScriptStruct = Cast<UScriptStruct>(Desc.ValueTypeObject);
 			break;
 			
-			// @todo_pcg: Support class and objects
+		// @todo_pcg: Support class and objects
 // 		case EPCGMetadataTypes::Class:
 // 			if (const UClass* Class = Cast<UClass>(Desc.ValueTypeObject))
 // 			{
 // 				FClassProperty* Prop = new FClassProperty(PropertyScope, Desc.Name, RF_Public);
 // #if WITH_EDITORONLY_DATA
-// 				Prop->SetMetaClass(Desc.MetaClass ? Desc.MetaClass.Get() : const_cast<UClass*>(Class));
+// 				Prop->SetMetaClass(const_cast<UClass*>(Desc.ValueTypeObject ? Cast<UClass>(Desc.ValueTypeObject) : Class));
 // #else
 // 				Prop->SetMetaClass(const_cast<UClass*>(Class));
 // #endif
@@ -220,7 +220,7 @@ namespace PCG::Private
 // 			{
 // 				FSoftClassProperty* Prop = new FSoftClassProperty(PropertyScope, Desc.Name, RF_Public);
 // #if WITH_EDITORONLY_DATA
-// 				Prop->SetMetaClass(Desc.MetaClass ? Desc.MetaClass.Get() : const_cast<UClass*>(Class));
+// 				Prop->SetMetaClass(const_cast<UClass*>(Desc.ValueTypeObject ? Cast<UClass>(Desc.ValueTypeObject) : Class));
 // #else
 // 				Prop->SetMetaClass(const_cast<UClass*>(Class));
 // #endif
@@ -733,12 +733,178 @@ void FPCGMetadataAttributeGeneric::Init()
 
 void FPCGMetadataAttributeGeneric::Flatten()
 {
-	// @todo_pcg: TODO
+	// Implementation notes:
+	// We don't need to flatten the EntryToValueKeyMap - this will have been taken care of in the metadata flatten
+
+	// Flatten values, from root to current attribute
+	if (Parent && UnderlyingProperty)
+	{
+		PCG::TUniqueScopeLock WriteLock(ValueLock);
+		
+		TArray<const FScriptArray*> OriginalValues;
+		
+		int32 ValueCount = 0;
+
+		const FPCGMetadataAttributeGeneric* Current = this;
+		while (Current)
+		{
+			if (ensure(Current->Values.Num() > 0))
+			{
+				ValueCount += Current->Values.Num();
+				OriginalValues.Add(&Current->Values);
+			}
+			
+			if (Current->Parent && Current->Parent->IsGeneric())
+			{
+				Current = static_cast<const FPCGMetadataAttributeGeneric*>(Current->Parent);
+			}
+			else
+			{
+				Current = nullptr;
+			}
+		}
+		
+		// Swap the current values to an empty one
+		FScriptArray NewValues{};
+		FScriptArrayHelper NewValuesHelper{UnderlyingProperty->Property.Get(), &NewValues};
+		
+		// Since each attribute in the hierarchy chain has a default value on index 0, all the values are offset by one
+		// To keep the mapping intact, we have to also copy the default value of each parent. But we'll make sure the default
+		// value (the one at index 0 at the end) is the one from `this`
+		if (UnderlyingProperty->bIsPlainOldData)
+		{
+			NewValuesHelper.AddUninitializedValues(ValueCount);
+		}
+		else
+		{
+			NewValuesHelper.AddValues(ValueCount);
+		}
+		
+		// Then copy all the values in reverse order (oldest first)
+		int32 Count = 0;
+		for (int32 ValuesIndex = OriginalValues.Num() - 1; ValuesIndex >= 0; --ValuesIndex)
+		{
+			FScriptArrayHelper OtherHelper{UnderlyingProperty->Property.Get(), OriginalValues[ValuesIndex]};
+			const int32 OtherCount = OtherHelper.Num();
+			UnderlyingProperty->Copy(NewValuesHelper.GetRawPtr(Count), OtherHelper.GetRawPtr(0), /*Count=*/OtherCount);
+			Count += OtherCount;
+		}
+		
+		// Then override the default value (index 0) by this default value
+		UnderlyingProperty->Copy(NewValues.GetData(), Values.GetData(), /*Count=*/1);
+		
+		// Finally move the new values in place of the old values.
+		if (TOptional<FScriptArrayHelper> Helper = MakeArrayHelper(); ensure(Helper.IsSet()))
+		{
+			Helper->MoveAssign(&NewValues);
+		}
+	}
+		
+	// Reset value offset, and lose parent
+	ValueKeyOffset = 0;
+	Parent = nullptr;
 }
 
 void FPCGMetadataAttributeGeneric::FlattenAndCompress(const TArrayView<const PCGMetadataEntryKey>& InEntryKeysToKeep)
 {
-	// @todo_pcg: TODO
+	// No entries, we can just delete everything in the attribute.
+	if (InEntryKeysToKeep.IsEmpty() || !UnderlyingProperty)
+	{
+		Reset();
+		return;
+	}
+
+	TArray<PCGMetadataValueKey> AllValueKeys;
+	TArray<PCGMetadataValueKey> AllUniqueValueKeys;
+	AllValueKeys.Reserve(InEntryKeysToKeep.Num());
+	bool bUseValueKeys = UsesValueKeys();
+	if (bUseValueKeys)
+	{
+		AllUniqueValueKeys.Reserve(InEntryKeysToKeep.Num());
+	}
+
+	// First gather all value keys associated with the entry keys to keep.
+	// If we compress data, we also store the unique value keys used (that is not default).
+	for (const PCGMetadataEntryKey& EntryKey : InEntryKeysToKeep)
+	{
+		AllValueKeys.Add(GetValueKey(EntryKey));
+		if (bUseValueKeys)
+		{
+			if (AllValueKeys.Last() != PCGDefaultValueKey)
+			{
+				AllUniqueValueKeys.AddUnique(AllValueKeys.Last());
+			}
+		}
+	}
+
+	// Then for each value key (or unique values keys), gather the value in a NewValues array
+	// and also keep a mapping between old value key and new value key.
+	// Only done if the old value key is not the default one.
+	
+	TMap<PCGMetadataValueKey, PCGMetadataValueKey> ValueKeyMapping;
+	FScriptArray NewValues{};
+	FScriptArrayHelper NewValuesHelper{UnderlyingProperty->Property.Get(), &NewValues};
+	
+	const TArray<PCGMetadataValueKey>& AllValueKeysRef = !bUseValueKeys ? AllValueKeys : AllUniqueValueKeys;
+	// Take into account the default value in the size.
+	if (UnderlyingProperty->bIsPlainOldData)
+	{
+		NewValuesHelper.AddUninitializedValues(AllValueKeysRef.Num() + 1);
+	}
+	else
+	{
+		NewValuesHelper.AddValues(AllValueKeysRef.Num() + 1);
+	}
+	
+	ValueKeyMapping.Reserve(AllValueKeysRef.Num());
+	
+	// Copy the default value
+	UnderlyingProperty->Copy(NewValues.GetData(), Values.GetData(), /*Count=*/1);
+	
+	int32 Count = 1;
+	for (PCGMetadataValueKey ValueKey : AllValueKeysRef)
+	{
+		if (ValueKey != PCGDefaultValueKey)
+		{
+			ValueKeyMapping.Add(ValueKey, Count);
+			
+			// @todo_pcg: Improve to find all contiguous values before copying.
+			const void* SrcPtr = GetReadAddressForValueKey_Unsafe(ValueKey);
+			UnderlyingProperty->Copy(NewValuesHelper.GetRawPtr(Count++), SrcPtr, /*Count=*/1);
+		}
+	}
+
+	// Move the new values in place of the old values.
+	if (TOptional<FScriptArrayHelper> Helper = MakeArrayHelper(); ensure(Helper.IsSet()))
+	{
+		PCG::TUniqueScopeLock WriteLock(ValueLock);
+		Helper->MoveAssign(&NewValues);
+	}
+
+	// And finally, create a new entry to value mapping.
+	// Logic is that each entry to keep will have their "index" as new entry key
+	// (like if the entries to keep are [25, 47, 54], the new entries would be [0, 1, 2]).
+	// So the operation is:
+	// All pairs Old EK -> Old VK transform to New EK -> New VK.
+	TMap<PCGMetadataEntryKey, PCGMetadataValueKey> NewMap;
+	for (int32 i = 0; i < InEntryKeysToKeep.Num(); ++i)
+	{
+		PCGMetadataValueKey ValueKey = AllValueKeys[i];
+		if (ValueKey != PCGDefaultValueKey && ensure(InEntryKeysToKeep[i] != PCGInvalidEntryKey))
+		{
+			NewMap.Add(i, ValueKeyMapping[ValueKey]);
+		}
+	}
+
+	// And move the map.
+	{
+		PCG::TUniqueScopeLock WriteLock(EntryMapLock);
+		EntryToValueKeyMap = MoveTemp(NewMap);
+	}
+
+	// At the end, reset value offset, and lose parent.
+	ValueKeyOffset = 0;
+	Parent = nullptr;
 }
 
 void FPCGMetadataAttributeGeneric::Reset()
@@ -1305,7 +1471,7 @@ const void* FPCGMetadataAttributeGeneric::GetReadAddressFromEntryKey_Unsafe(cons
 
 void FPCGMetadataAttributeGeneric::SetValueFromProperty(PCGMetadataEntryKey ItemKey, const void* SrcPtr, const FProperty* Property)
 {
-	if (!SrcPtr || !Property || ItemKey == PCGInvalidEntryKey)
+	if (!SrcPtr || !Property)
 	{
 		return;
 	}
@@ -1317,7 +1483,9 @@ void FPCGMetadataAttributeGeneric::SetValueFromProperty(PCGMetadataEntryKey Item
 			return;
 		}
 		
-		SetValuesFromValueKeys({ItemKey}, AddValues_Internal(FInValues{TInPlaceType<FInValuesByValue>{}, SrcPtr, 1}, 1));
+		const int32 ExistingIndex = ItemKey == PCGInvalidEntryKey && ensure(Values.Num() > 0) ? 0 : INDEX_NONE;
+
+		SetValuesFromValueKeys({ItemKey}, AddValues_Internal(FInValues{TInPlaceType<FInValuesByValue>{}, SrcPtr, 1}, 1, ExistingIndex), /*bResetDefaultValueOnDefaultValueKey=*/ItemKey==PCGInvalidEntryKey);
 	}
 }
 
@@ -1353,6 +1521,15 @@ int32 FPCGMetadataAttributeGeneric::PreallocateValues(TArrayView<PCGMetadataEntr
 	return StartIndex;
 }
 
+void FPCGMetadataAttributeGeneric::Prepare(int32 Count)
+{
+	EntryToValueKeyMap.Reserve(EntryToValueKeyMap.Num() + Count);
+	if (UnderlyingProperty && !UnderlyingProperty->bCompressData)
+	{
+		AllocateValues_Unsafe(Count);
+	}
+}
+
 int32 FPCGMetadataAttributeGeneric::AllocateValues_Unsafe(int32 Count)
 {
 	if (TOptional<FScriptArrayHelper> Helper = MakeArrayHelper(); Helper.IsSet())
diff --git a/Engine/Plugins/PCG/Source/PCG/Public/Metadata/PCGMetadataAttributeGeneric.h b/Engine/Plugins/PCG/Source/PCG/Public/Metadata/PCGMetadataAttributeGeneric.h
--- a/Engine/Plugins/PCG/Source/PCG/Public/Metadata/PCGMetadataAttributeGeneric.h
+++ b/Engine/Plugins/PCG/Source/PCG/Public/Metadata/PCGMetadataAttributeGeneric.h
@@ -227,6 +227,14 @@ public:
 		GetValues(ValueKeys, OutValues);
 	}
 
+	template <typename T>
+	T GetValueFromItemKey(const PCGMetadataEntryKey InEntryKey) const
+	{
+		T OutValue{};
+		GetValuesFromItemKeys(MakeArrayView(&InEntryKey, 1), MakeArrayView(&OutValue, 1));
+		return OutValue;
+	}
+
 	// Override the default value for this attribute.
 	// Will fail if the type mismatch.
 	template <typename T>
@@ -239,6 +247,12 @@ public:
 		}
 	}
 
+	template <typename T>
+	void SetValue(const PCGMetadataEntryKey InEntryKey, const T& InValue)
+	{
+		SetValuesFromValueKeys(MakeArrayView(&InEntryKey, 1), AddValues<T>({InValue}));
+	}
+
 	template <typename T>
 	void SetValues(TArrayView<const PCGMetadataEntryKey * const> ItemKeys, TArrayView<const T> InValues)
 	{
@@ -286,35 +300,32 @@ public:
 	}
 
 	// Copy values from SrcPtr to the attribute. Property must match the attribute property, and the SrcPtr is expected to be
-	// reinterpreted as the property type.
+	// reinterpreted as the property type. The ItemKey can be the invalid entry key to set the default value.
 	PCG_API void SetValueFromProperty(PCGMetadataEntryKey ItemKey, const void* SrcPtr, const FProperty* Property);
 	
-	int32 PreallocateValues(TArrayView<PCGMetadataEntryKey*> EntryKeys, bool bLockless);
-
-	// Allocate the values (without initialization if it can) and return the first index of that new memory space.
-	// Return INDEX_NONE if it didn't work.
-	int32 AllocateValues_Unsafe(int32 Count);
+	PCG_API int32 PreallocateValues(TArrayView<PCGMetadataEntryKey*> EntryKeys, bool bLockless);
 	
-	void Prepare(int32 Count)
-	{
-		// @todo_pcg To be implemented
-	}
+	PCG_API void Prepare(int32 Count);
 
 	// @todo_pcg: Temporary for debug. Build a struct that will contain the value at the entry key.
 	PCG_API FInstancedPropertyBag BuildStructForDebug(PCGMetadataEntryKey EntryKey) const;
 
+	// Utility function to get the pointer to a value using a value key. Return nullptr if invalid. Advanced function, use at your own risk.
+	PCG_API const void* GetReadAddressForValueKey_Unsafe(const PCGMetadataValueKey InValueKey) const;
+
+	// Utility function to get the pointer to a value using an entry key. Return nullptr if invalid. Advanced function, use at your own risk.
+	PCG_API const void* GetReadAddressFromEntryKey_Unsafe(const PCGMetadataEntryKey InEntryKey) const;
+
 protected:
+	// Allocate the values (without initialization if it can) and return the first index of that new memory space.
+	// Return INDEX_NONE if it didn't work.
+	int32 AllocateValues_Unsafe(int32 Count);
+
 	// Used to make an array helper of the local FScriptArray.
 	// Optional, since the attribute can be ill-formed, so the user should always check if it is set before
 	// using it.
 	TOptional<FScriptArrayHelper> MakeArrayHelper() const;
 	
-	// Utility function to get the pointer to a value using a value key. Return nullptr if invalid.
-	const void* GetReadAddressForValueKey_Unsafe(const PCGMetadataValueKey InValueKey) const;
-
-	// Utility function to get the pointer to a value using an entry key. Return nullptr if invalid.
-	const void* GetReadAddressFromEntryKey_Unsafe(const PCGMetadataEntryKey InEntryKey) const;
-	
 	/**
 	 * GetValues implementation section
 	 * GetValues_Internal is working on a Variant that will type erase the underlying output container
diff --git a/Engine/Plugins/PCG/Tests/Private/Metadata/PCGMetadataFlattenTests.cpp b/Engine/Plugins/PCG/Tests/Private/Metadata/PCGMetadataFlattenTests.cpp
--- a/Engine/Plugins/PCG/Tests/Private/Metadata/PCGMetadataFlattenTests.cpp
+++ b/Engine/Plugins/PCG/Tests/Private/Metadata/PCGMetadataFlattenTests.cpp
@@ -19,6 +19,7 @@ TEST_CASE_METHOD(PCGTests::FPCGBaseTestWithBasicContext, "PCG::Metadata::Flatten
 
 	constexpr int32 NumPoints = 10;
 	UPCGPointArrayData* RootPointData = CreateData<UPCGPointArrayData>();
+	RootPointData->AllocateProperties(EPCGPointNativeProperties::MetadataEntry);
 	RootPointData->SetNumPoints(NumPoints);
 	
 	FPCGMetadataAttribute<float>* Attribute1 = RootPointData->Metadata->CreateAttribute<float>(Attribute1Name, -0.1f, true, true);
@@ -185,6 +186,181 @@ TEST_CASE_METHOD(PCGTests::FPCGBaseTestWithBasicContext, "PCG::Metadata::Flatten
 	}
 }
 
+TEST_CASE_METHOD(PCGTests::FPCGBaseTestWithBasicContext, "PCG::Metadata::FlattenGeneric", "[PCG][Metadata][Flatten]")
+{
+	static const FName Attribute1Name = TEXT("FloatAttr");
+	static const FName Attribute2Name = TEXT("StringAttr");
+	static const FName Attribute3Name = TEXT("IntAttr");
+
+	constexpr int32 NumPoints = 10;
+	UPCGPointArrayData* RootPointData = CreateData<UPCGPointArrayData>();
+	RootPointData->AllocateProperties(EPCGPointNativeProperties::MetadataEntry);
+	RootPointData->SetNumPoints(NumPoints);
+	
+	FPCGMetadataAttributeGeneric* Attribute1 = RootPointData->MutableMetadata()->GetDefaultMetadataDomain()->CreateGenericAttribute<float>(Attribute1Name, -0.1f);
+	FPCGMetadataAttributeGeneric* Attribute2 = RootPointData->MutableMetadata()->GetDefaultMetadataDomain()->CreateGenericAttribute<FString>(Attribute2Name, TEXT("Default"));
+
+	check(Attribute1 && Attribute2);
+
+	TPCGValueRange<int64> MetadataEntryRange = RootPointData->GetMetadataEntryValueRange();
+	for (int32 i = 0; i < RootPointData->GetNumPoints(); ++i)
+	{
+		RootPointData->MutableMetadata()->InitializeOnSet(MetadataEntryRange[i]);
+		Attribute1->SetValue<float>(MetadataEntryRange[i], static_cast<float>(i) * 0.1f);
+
+		if (i % 2 == 0)
+		{
+			Attribute2->SetValue<FString>(MetadataEntryRange[i], FString::Printf(TEXT("%d"), i % 4));
+		}
+	}
+
+	// At the end of the first set, metadata has 10 entries, and values for each points are
+	// for Attribute 1: [0.0, 0.1, 0.2, 0.3, 0.4, 0.5, 0.6, 0.7, 0.8, 0.9]
+	// for Attribute 2: ["0", "Default", "2", "Default", "0", ...]
+	SECTION("RootMetadata has the right number of entries")
+	{
+		REQUIRE_EQUAL(RootPointData->Metadata->GetItemCountForChild(), NumPoints);
+	}
+
+	UPCGBasePointData* FirstChildPointData = Cast<UPCGBasePointData>(RootPointData->DuplicateData(nullptr));
+	Attribute1 = FirstChildPointData->MutableMetadata()->GetDefaultMetadataDomain()->GetMutableGenericAttribute(Attribute1Name);
+	Attribute2 = FirstChildPointData->MutableMetadata()->GetDefaultMetadataDomain()->GetMutableGenericAttribute(Attribute2Name);
+
+	SECTION("Attributes exists in first child")
+	{
+		REQUIRE_NOT_EQUAL(Attribute1, nullptr);
+		REQUIRE_NOT_EQUAL(Attribute2, nullptr);
+	}
+	
+	check(Attribute1 && Attribute2);
+
+	// Override all the values for Attribute 1, and replace all "1" by "0" in attribute 2
+	// Also add a third attribute
+	FPCGMetadataAttributeGeneric* Attribute3 = FirstChildPointData->MutableMetadata()->GetDefaultMetadataDomain()->CreateGenericAttribute<int>(Attribute3Name, -1);
+	check(Attribute3);
+
+	SECTION("First child has right number of points")
+	{
+		REQUIRE_EQUAL(FirstChildPointData->GetNumPoints(), NumPoints);
+	}
+
+	TPCGValueRange<int64> FirstChildMetadataEntryRange = FirstChildPointData->GetMetadataEntryValueRange();
+	for (int32 i = 0; i < FirstChildPointData->GetNumPoints(); ++i)
+	{
+		FirstChildPointData->Metadata->InitializeOnSet(FirstChildMetadataEntryRange[i]);
+		Attribute1->SetValue<float>(FirstChildMetadataEntryRange[i], static_cast<float>(i) * 1.1f);
+
+		if (i % 2 == 0)
+		{
+			Attribute2->SetValue<FString>(FirstChildMetadataEntryRange[i], TEXT("0"));
+		}
+
+		Attribute3->SetValue<int>(FirstChildMetadataEntryRange[i], i);
+	}
+
+	// At the end of the second set, metadata has 20 entries, and values for each points are
+	// for Attribute 1: [0.0, 1.1, 2.2, 3.3, 4.4, 5.5, 6.6, 7.7, 8.8, 9.9]
+	// for Attribute 2: ["0", "Default", "0", "Default", "0", ...]
+	// for Attribute 3: [0, 1, 2, 3, 4, 5, 6, 7, 8, 9]
+	SECTION("First child metadata has right number of entries")
+	{
+		REQUIRE_EQUAL(FirstChildPointData->Metadata->GetItemCountForChild(), 2 * NumPoints);
+	}
+	
+
+	// Second child is keeping the metadata entry for even numbers and reset metadata entry for the rest (point back to default). Also override Attribute 3 values for even numbers
+	UPCGBasePointData* SecondChildPointData = CreateData<UPCGPointArrayData>();
+
+	FPCGInitializeFromDataParams InitializeFromDataParams(FirstChildPointData);
+	InitializeFromDataParams.bInheritSpatialData = false;
+	SecondChildPointData->InitializeFromDataWithParams(InitializeFromDataParams);
+	
+	Attribute3 = SecondChildPointData->MutableMetadata()->GetDefaultMetadataDomain()->GetMutableGenericAttribute(Attribute3Name);
+	SECTION("Attributes exists in second child")
+	{
+		REQUIRE_NOT_EQUAL(Attribute3, nullptr);
+	}
+
+	check(Attribute3);
+
+	SecondChildPointData->SetNumPoints(FirstChildPointData->GetNumPoints());
+	UPCGBasePointData::SetPoints(FirstChildPointData, SecondChildPointData, {}, /*bCopyAll=*/true);
+
+	TPCGValueRange<int64> SecondChildMetadataEntryRange = SecondChildPointData->GetMetadataEntryValueRange();
+	for (int32 i = 0; i < SecondChildPointData->GetNumPoints(); ++i)
+	{
+		if (i % 2 == 0)
+		{
+			SecondChildPointData->Metadata->InitializeOnSet(SecondChildMetadataEntryRange[i]);
+			Attribute3->SetValue<int>(SecondChildMetadataEntryRange[i], 10 * i);
+		}
+		else
+		{
+			SecondChildMetadataEntryRange[i] = PCGInvalidEntryKey;
+		}
+	}
+
+	// At the end of the third set, metadata has 25 entries, and values for each points are:
+	// for Attribute 1: [0.0, -0.1, 2.2, -0.1, 4.4, -0.1, 6.6, -0.1, 8.8, -0.1]
+	// for Attribute 2: ["0", "Default", "0", "Default", "0", ...]
+	// for Attribute 3: [0, -1, 20, -1, 40, -1, 60, -1, 80, -1]
+	SECTION("Second Child Metadata has the right number of entries")
+	{
+		REQUIRE_EQUAL(SecondChildPointData->Metadata->GetItemCountForChild(), 2 * NumPoints + NumPoints / 2);
+	}
+	
+
+	// For final set, duplicate the data and flatten it
+	UPCGBasePointData* FinalPointData = Cast<UPCGBasePointData>(SecondChildPointData->DuplicateData(nullptr));
+	FinalPointData->Flatten();
+
+	Attribute1 = FinalPointData->MutableMetadata()->GetDefaultMetadataDomain()->GetMutableGenericAttribute(Attribute1Name);
+	Attribute2 = FinalPointData->MutableMetadata()->GetDefaultMetadataDomain()->GetMutableGenericAttribute(Attribute2Name);
+	Attribute3 = FinalPointData->MutableMetadata()->GetDefaultMetadataDomain()->GetMutableGenericAttribute(Attribute3Name);
+
+	SECTION("Attributes exists in final child")
+	{
+		REQUIRE_NOT_EQUAL(Attribute1, nullptr);
+		REQUIRE_NOT_EQUAL(Attribute2, nullptr);
+		REQUIRE_NOT_EQUAL(Attribute3, nullptr);
+	}
+	
+	check(Attribute1 && Attribute2 && Attribute3);
+
+	SECTION("Flatten metadata has the right number of entries/values")
+	{
+		// + 1 for the default value
+		REQUIRE_EQUAL(FinalPointData->Metadata->GetItemCountForChild(), 5);
+		REQUIRE_EQUAL(Attribute1->GetValueKeyOffsetForChild(), 5 + 1);
+		REQUIRE_EQUAL(Attribute2->GetValueKeyOffsetForChild(), 1 + 1);
+		REQUIRE_EQUAL(Attribute3->GetValueKeyOffsetForChild(), 5 + 1);
+		REQUIRE_EQUAL(FinalPointData->GetNumPoints(), NumPoints);
+	}
+
+	// Validate the values
+	const TConstPCGValueRange<int64> FinalMetadataEntryRange = FinalPointData->GetConstMetadataEntryValueRange();
+	auto i = GENERATE_COPY(range(0, NumPoints));
+	
+	SECTION("Validate values")
+	{
+		const int64 MetadataEntry = FinalMetadataEntryRange[i];
+		if (i % 2 == 0)
+		{
+			REQUIRE_EQUAL(MetadataEntry, static_cast<PCGMetadataEntryKey>(i / 2));
+			REQUIRE_EQUAL(Attribute1->GetValueFromItemKey<float>(MetadataEntry), i * 1.1f);
+			REQUIRE_EQUAL(Attribute2->GetValueFromItemKey<FString>(MetadataEntry), TEXT("0"));
+			REQUIRE_EQUAL(Attribute3->GetValueFromItemKey<int>(MetadataEntry), i * 10);
+		}
+		else
+		{
+			REQUIRE_EQUAL(MetadataEntry, PCGInvalidEntryKey);
+			REQUIRE_EQUAL(Attribute1->GetValueFromItemKey<float>(MetadataEntry), -0.1f);
+			REQUIRE_EQUAL(Attribute2->GetValueFromItemKey<FString>(MetadataEntry), TEXT("Default"));
+			REQUIRE_EQUAL(Attribute3->GetValueFromItemKey<int>(MetadataEntry), -1);
+		}
+	}
+}
+
 /**
  * Flattening metadata that has no attribute still must be flattened. Same for Data domain.
  */
@@ -196,6 +372,7 @@ TEST_CASE_METHOD(PCGTests::FPCGBaseTestWithBasicContext, "PCG::Metadata::Flatten
 	// Preparing the data
 	constexpr int32 NumPoints = 1;
 	UPCGPointArrayData* PointArrayData = CreateData<UPCGPointArrayData>();
+	PointArrayData->AllocateProperties(EPCGPointNativeProperties::MetadataEntry);
 	PointArrayData->SetNumPoints(NumPoints);
 	
 	FPCGMetadataDomain* DataDomain = PointArrayData->MutableMetadata()->GetMetadataDomain(PCGMetadataDomainID::Data);
@@ -238,4 +415,77 @@ TEST_CASE_METHOD(PCGTests::FPCGBaseTestWithBasicContext, "PCG::Metadata::Flatten
 	
 	REQUIRE_EQUAL(ChildElementDomain->GetAttributeCount(), 0);
 	REQUIRE_EQUAL(ChildElementDomain->GetParent(), nullptr);
+}
+
+TEST_CASE_METHOD(PCGTests::FPCGBaseTestWithBasicContext, "PCG::Metadata::FlattenWithoutCompression", "[PCG][Metadata][Flatten]")
+{
+	static const FName Attribute1Name = TEXT("FloatAttr");
+	static const FName Attribute2Name = TEXT("StringAttr");
+
+	constexpr int32 NumPoints = 5;
+	UPCGPointArrayData* RootPointData = CreateData<UPCGPointArrayData>();
+	RootPointData->AllocateProperties(EPCGPointNativeProperties::MetadataEntry);
+	RootPointData->SetNumPoints(NumPoints);
+	
+	FPCGMetadataAttributeGeneric* Attribute1 = RootPointData->MutableMetadata()->GetDefaultMetadataDomain()->CreateGenericAttribute<float>(Attribute1Name, -0.1f);
+	FPCGMetadataAttributeGeneric* Attribute2 = RootPointData->MutableMetadata()->GetDefaultMetadataDomain()->CreateGenericAttribute<FString>(Attribute2Name, TEXT("Default"));
+	
+	REQUIRE_NOT_EQUAL(Attribute1, nullptr);
+	REQUIRE_NOT_EQUAL(Attribute2, nullptr);
+	
+	TArray<float> FloatValues;
+	TArray<FString> StringValues;
+	
+	for (int32 i = 0; i < NumPoints * 2; ++i)
+	{
+		FloatValues.Emplace(i);
+		StringValues.Emplace(FString::Printf(TEXT("%d"), i % 3));
+	}
+	
+	TPCGValueRange<int64> MetadataEntryRange = RootPointData->GetMetadataEntryValueRange();
+	for (int32 i = 0; i < NumPoints; ++i)
+	{
+		RootPointData->MutableMetadata()->InitializeOnSet(MetadataEntryRange[i]);
+		Attribute1->SetValue(MetadataEntryRange[i], FloatValues[i]);
+		Attribute2->SetValue(MetadataEntryRange[i], StringValues[i]);
+	}
+	
+	UPCGPointArrayData* ChildPointData = CreateData<UPCGPointArrayData>();
+	ChildPointData->InitializeFromDataWithParams(FPCGInitializeFromDataParams{RootPointData});
+	ChildPointData->SetNumPoints(NumPoints * 2);
+	
+	Attribute1 = ChildPointData->MutableMetadata()->GetDefaultMetadataDomain()->GetMutableGenericAttribute(Attribute1Name);
+	Attribute2 = ChildPointData->MutableMetadata()->GetDefaultMetadataDomain()->GetMutableGenericAttribute(Attribute2Name);
+	
+	REQUIRE_NOT_EQUAL(Attribute1, nullptr);
+	REQUIRE_NOT_EQUAL(Attribute2, nullptr);
+	
+	TPCGValueRange<int64> ChildMetadataEntryRange = ChildPointData->GetMetadataEntryValueRange();
+	for (int32 i = NumPoints; i < 2 * NumPoints; ++i)
+	{
+		ChildPointData->MutableMetadata()->InitializeOnSet(ChildMetadataEntryRange[i]);
+		Attribute1->SetValue(ChildMetadataEntryRange[i], FloatValues[i]);
+		Attribute2->SetValue(ChildMetadataEntryRange[i], StringValues[i]);
+	}
+	
+	REQUIRE_EQUAL(Attribute1->GetNumberOfEntries(), NumPoints);
+	REQUIRE_EQUAL(Attribute2->GetNumberOfEntries(), NumPoints);
+	
+	ChildPointData->MutableMetadata()->GetDefaultMetadataDomain()->FlattenImpl();
+	
+	REQUIRE_EQUAL(Attribute1->GetParent(), nullptr);
+	REQUIRE_EQUAL(Attribute2->GetParent(), nullptr);
+	
+	REQUIRE_EQUAL(Attribute1->GetNumberOfEntries(), 2 * NumPoints);
+	REQUIRE_EQUAL(Attribute2->GetNumberOfEntries(), 2 * NumPoints);
+	
+	REQUIRE_EQUAL(Attribute1->GetValueFromItemKey<float>(PCGInvalidEntryKey), -0.1f);
+	REQUIRE_EQUAL(Attribute2->GetValueFromItemKey<FString>(PCGInvalidEntryKey), TEXT("Default"));
+	
+	for (int32 i = 0; i < 2 * NumPoints; ++i)
+	{
+		CAPTURE(i);
+		REQUIRE_EQUAL(Attribute1->GetValueFromItemKey<float>(ChildMetadataEntryRange[i]), FloatValues[i]);
+		REQUIRE_EQUAL(Attribute2->GetValueFromItemKey<FString>(ChildMetadataEntryRange[i]), StringValues[i]);
+	}
 }
\ No newline at end of file
