[PCG] PCG Indexing Improvements

* FPCGIndexCollection
+ New constructor
* Binary search for index lookup
+ New AddIndex helper for single index additions

#rnx

#Synced-CL 50977214

diff --git a/Engine/Plugins/PCG/Source/PCG/Private/Helpers/Parsing/PCGIndexing.cpp b/Engine/Plugins/PCG/Source/PCG/Private/Helpers/Parsing/PCGIndexing.cpp
--- a/Engine/Plugins/PCG/Source/PCG/Private/Helpers/Parsing/PCGIndexing.cpp
+++ b/Engine/Plugins/PCG/Source/PCG/Private/Helpers/Parsing/PCGIndexing.cpp
@@ -2,6 +2,8 @@
 
 #include "Helpers/Parsing/PCGIndexing.h"
 
+#include "Algo/MaxElement.h"
+
 namespace PCGIndexing
 {
 	bool FPCGIndexRange::ContainsIndex(const int32 Index) const
@@ -29,6 +31,17 @@ namespace PCGIndexing
 		return !operator==(OtherRange);
 	}
 
+	FPCGIndexCollection::FPCGIndexCollection(const TArrayView<const int32> InArrayView)
+	{
+		AddIndices(InArrayView, /*bGrowIfNeeded=*/true);
+	}
+
+	FPCGIndexCollection::FPCGIndexCollection(const TArrayView<const int32> InArrayView, const int32 InArraySize)
+		: ArraySize(FMath::Max(InArraySize, 0))
+	{
+		AddIndices(InArrayView);
+	}
+
 	bool FPCGIndexCollection::AddRange(const int32 StartIndex, const int32 EndIndex)
 	{
 		if (!ensure(IsValid()))
@@ -87,29 +100,23 @@ namespace PCGIndexing
 	bool FPCGIndexCollection::RangeIsValid(const FPCGIndexRange& Range) const
 	{
 		return IsValid()
-				&& Range.StartIndex >= 0
-				&& Range.StartIndex <= ArraySize
-				&& Range.EndIndex >= 0
-				&& Range.EndIndex <= ArraySize
-				&& Range.StartIndex < Range.EndIndex;
+			&& Range.StartIndex >= 0
+			&& Range.StartIndex <= ArraySize
+			&& Range.EndIndex >= 0
+			&& Range.EndIndex <= ArraySize
+			&& Range.StartIndex < Range.EndIndex;
 	}
 
-	bool FPCGIndexCollection::ContainsIndex(const int32 Index) const
+	bool FPCGIndexCollection::ContainsIndex(const int32 InIndex) const
 	{
 		if (!IsValid())
 		{
 			return false;
 		}
 
-		for (const FPCGIndexRange& Range : IndexRanges)
-		{
-			if (Range.ContainsIndex(Index))
-			{
-				return true;
-			}
-		}
-
-		return false;
+		// Employ a binary search since the collection is already sorted - O(log n) lookup
+		const int32 LowerIndex = Algo::LowerBoundBy(IndexRanges, InIndex, [](const FPCGIndexRange& Range) { return Range.EndIndex; });
+		return LowerIndex < IndexRanges.Num() && IndexRanges[LowerIndex].ContainsIndex(InIndex);
 	}
 
 	int32 FPCGIndexCollection::GetArraySize() const
@@ -145,6 +152,14 @@ namespace PCGIndexing
 		return IndexRanges.IsEmpty();
 	}
 
+	void FPCGIndexCollection::ForEach(const TFunctionRef<void(FPCGIndexRange)> Action)
+	{
+		for (const FPCGIndexRange& Range : IndexRanges)
+		{
+			Action(Range);
+		}
+	}
+
 	bool FPCGIndexCollection::operator==(const FPCGIndexCollection& Other) const
 	{
 		return (ArraySize == Other.ArraySize) && (IndexRanges == Other.IndexRanges);
@@ -161,6 +176,12 @@ namespace PCGIndexing
 		return *this;
 	}
 
+	FPCGIndexCollection& FPCGIndexCollection::operator+=(const TArrayView<const int32>& InArrayView)
+	{
+		AddIndices(InArrayView, /*bGrowIfNeeded=*/true);
+		return *this;
+	}
+
 	FPCGIndexRange FPCGIndexCollection::AdjustIndicesAndCreateRange(int32 StartIndex, int32 EndIndex) const
 	{
 		// Convert both indices to a positive range
@@ -187,4 +208,75 @@ namespace PCGIndexing
 
 		return FPCGIndexRange(MergedStart, MergedEnd);
 	}
-}
+
+	void FPCGIndexCollection::AddIndices(const TArrayView<const int32> InArrayView, const bool bGrowIfNeeded)
+	{
+		// @todo_pcg: Consider a divide and conquer approach for very large merges.
+		if (!InArrayView.IsEmpty())
+		{
+			// If the indices contained in the array view are larger than the current array size, expand
+			if (bGrowIfNeeded)
+			{
+				if (const int32* InMax = Algo::MaxElement(InArrayView); InMax && *InMax > 0)
+				{
+					ArraySize = FMath::Max(*InMax + 1, ArraySize);
+				}
+			}
+
+			if (!IsValid())
+			{
+				// Edge Case: Negative indices only, which require the ArraySize to be positive.
+				return;
+			}
+
+			// Remove invalid indices from the inbound array.
+			TArray<int32, TInlineAllocator<128>> Indices(InArrayView);
+			for (int32 i = Indices.Num() - 1; i >= 0; --i)
+			{
+				Indices[i] = (Indices[i] >= 0 ? Indices[i] : ArraySize + Indices[i]);
+				// Invalid negative index... discard
+				if (Indices[i] < 0)
+				{
+					Indices.RemoveAtSwap(i, EAllowShrinking::No);
+				}
+			}
+
+			if (Indices.IsEmpty())
+			{
+				return;
+			}
+
+			// Sort the indices to make the range compilation easier.
+			Indices.Sort();
+
+			int32 RangeStart = Indices[0];
+			int32 RangeEnd = RangeStart;
+
+			// For the remaining indices, walk through and distill them into ranges.
+			for (int32 i = 1; i < Indices.Num(); ++i)
+			{
+				// Skip duplicate indices (ranges can't contain them, but an array of indices might).
+				if (Indices[i] <= RangeEnd)
+				{
+					continue;
+				}
+
+				// Grow the working range, if needed.
+				if (Indices[i] == RangeEnd + 1)
+				{
+					RangeEnd = Indices[i];
+				}
+				else if (Indices[i] > RangeEnd + 1)
+				{
+					// Range complete - add to the collection and start a new one.
+					AddRange(RangeStart, RangeEnd + 1);
+					RangeStart = Indices[i];
+					RangeEnd = Indices[i];
+				}
+			}
+
+			// Add the final range
+			AddRange(RangeStart, RangeEnd + 1);
+		}
+	}
+} // namespace PCGIndexing
diff --git a/Engine/Plugins/PCG/Source/PCG/Public/Helpers/Parsing/PCGIndexing.h b/Engine/Plugins/PCG/Source/PCG/Public/Helpers/Parsing/PCGIndexing.h
--- a/Engine/Plugins/PCG/Source/PCG/Public/Helpers/Parsing/PCGIndexing.h
+++ b/Engine/Plugins/PCG/Source/PCG/Public/Helpers/Parsing/PCGIndexing.h
@@ -7,7 +7,7 @@
 /** Helper namespace for handling index based evaluations. */
 namespace PCGIndexing
 {
-	/** A simple data structure to represent a range of indices [X,Y]. Validating ranges are client responsibility. */
+	/** A simple data structure to represent a range of indices [X,Y). Validating ranges are client responsibility. */
 	struct FPCGIndexRange
 	{
 		FPCGIndexRange(const int32 InStartIndex, const int32 InEndIndex) : StartIndex(InStartIndex), EndIndex(InEndIndex) {}
@@ -33,8 +33,13 @@ namespace PCGIndexing
 		static FPCGIndexCollection Invalid() { return {}; }
 
 		/** The constructor must accept the size of the array to support negative terminating indices. */
-		explicit FPCGIndexCollection(const int32 InArraySize) : ArraySize(FMath::Max(InArraySize, static_cast<int32>(INDEX_NONE))) {}
+		explicit FPCGIndexCollection(const int32 InArraySize) : ArraySize(FMath::Max(InArraySize, -1)) {}
 
+		/** Direct conversion from an array view of indices. Size is implicitly calculated from max value. */
+		explicit FPCGIndexCollection(const TArrayView<const int32> InArrayView);
+
+		/** Direct conversion from an array view of indices. Size is explicit. */
+		FPCGIndexCollection(const TArrayView<const int32> InArrayView, int32 InArraySize);
 
 		/** Add a new index range to the collection directly via start and end indices. */
 		bool AddRange(int32 StartIndex, int32 EndIndex);
@@ -46,7 +51,7 @@ namespace PCGIndexing
 		bool RangeIsValid(const FPCGIndexRange& Range) const;
 
 		/** Returns true if this collection contains the given index. */
-		bool ContainsIndex(int32 Index) const;
+		bool ContainsIndex(int32 InIndex) const;
 
 		/** Returns the abstract size of the array tied to the collection. */
 		int32 GetArraySize() const;
@@ -63,8 +68,12 @@ namespace PCGIndexing
 		/** The index collection has no ranges. */
 		bool IsEmpty() const;
 
+		/** Iterate over the entire range of indices. */
+		void ForEach(const TFunctionRef<void(FPCGIndexRange)> Action);
+
 		bool operator==(const FPCGIndexCollection& Other) const;
 		FPCGIndexCollection& operator+=(const FPCGIndexCollection& Other);
+		FPCGIndexCollection& operator+=(const TArrayView<const int32>& InArrayView);
 
 	private:
 		FPCGIndexCollection() = default;
@@ -78,8 +87,11 @@ namespace PCGIndexing
 		/** Returns the resultant range after merging two overlapping ranges. */
 		FPCGIndexRange MergeRanges(const FPCGIndexRange& FirstRange, const FPCGIndexRange& SecondRange) const;
 
+		/** Helper to add an array of indices. If marked to grow, the max value will be used to determine the new array size. */
+		void AddIndices(const TArrayView<const int32> InArrayView, bool bGrowIfNeeded = false);
+
 		/** The size of the representative array associated with this collection. Needed for negative terminating indices. */
-		int32 ArraySize = INDEX_NONE;
+		int32 ArraySize = 0;
 
 		/** A collection of the abstract index ranges in the collection. */
 		TArray<FPCGIndexRange, TInlineAllocator<64>> IndexRanges;
diff --git a/Engine/Plugins/PCG/Tests/Private/PCGIndexingTests.cpp b/Engine/Plugins/PCG/Tests/Private/PCGIndexingTests.cpp
new file mode 100644
--- /dev/null
+++ b/Engine/Plugins/PCG/Tests/Private/PCGIndexingTests.cpp
@@ -0,0 +1,199 @@
+// Copyright Epic Games, Inc. All Rights Reserved.
+
+#include "CoreMinimal.h"
+#include "TestHarness.h"
+
+#include <catch2/generators/catch_generators.hpp>
+
+#include "Helpers/Parsing/PCGParsing.h"
+#include "PCGTestsCommon.h"
+
+TEST_CASE_METHOD(PCGTests::FPCGBaseTestWithBasicContext, "PCG::Indexing::Basic", "[PCG][Indexing]")
+{
+	using PCGIndexing::FPCGIndexCollection;
+
+	SECTION("Empty array")
+	{
+		FPCGIndexCollection Collection(10);
+		TArray<int32> EmptyIndices;
+		Collection += EmptyIndices;
+		REQUIRE_EQUAL(Collection.GetTotalIndexCount(), 0);
+	}
+
+	SECTION("Empty array construction")
+	{
+		TArray<int32> EmptyIndices;
+		FPCGIndexCollection Collection(EmptyIndices, 10);
+		REQUIRE(Collection.IsValid());
+		REQUIRE(Collection.IsEmpty());
+		REQUIRE_EQUAL(Collection.GetArraySize(), 10);
+	}
+
+	SECTION("Array construction")
+	{
+		TArray<int32> Indices = {0, 1, 2, 3, 4};
+		FPCGIndexCollection Collection(Indices, 10);
+		REQUIRE_EQUAL(Collection.GetTotalIndexCount(), 5);
+		REQUIRE_EQUAL(Collection.GetTotalRangeCount(), 1);
+		REQUIRE_EQUAL(Collection.GetArraySize(), 10);
+	}
+
+	SECTION("Single index addition")
+	{
+		FPCGIndexCollection Collection(10);
+		Collection += {5};
+		REQUIRE_EQUAL(Collection.GetTotalIndexCount(), 1);
+		REQUIRE(Collection.ContainsIndex(5));
+		REQUIRE_FALSE(Collection.ContainsIndex(6));
+	}
+
+	SECTION("Single index via array")
+	{
+		FPCGIndexCollection Collection(10);
+		TArray<int32> Indices = {5};
+		Collection += Indices;
+		REQUIRE_EQUAL(Collection.GetTotalIndexCount(), 1);
+		REQUIRE(Collection.ContainsIndex(5));
+	}
+
+	SECTION("Consecutive indices merge")
+	{
+		FPCGIndexCollection Collection(10);
+		Collection += {3, 4, 5};
+		REQUIRE_EQUAL(Collection.GetTotalIndexCount(), 3);
+		REQUIRE_EQUAL(Collection.GetTotalRangeCount(), 1);
+	}
+
+	SECTION("Multiple, non-consecutive")
+	{
+		FPCGIndexCollection Collection(10);
+		Collection += {2, 4, 6};
+		REQUIRE_EQUAL(Collection.GetTotalIndexCount(), 3);
+		REQUIRE_EQUAL(Collection.GetTotalRangeCount(), 3);
+		REQUIRE((Collection.ContainsIndex(2) && Collection.ContainsIndex(4) && Collection.ContainsIndex(6)));
+		REQUIRE((!Collection.ContainsIndex(1) && !Collection.ContainsIndex(5) && !Collection.ContainsIndex(7)));
+	}
+
+	SECTION("Sorted, consecutive indices")
+	{
+		FPCGIndexCollection Collection(20);
+		TArray<int32> Indices = {0, 1, 2, 3, 4, 5, 6, 7, 8, 9};
+		Collection += Indices;
+		REQUIRE_EQUAL(Collection.GetTotalIndexCount(), 10);
+		REQUIRE_EQUAL(Collection.GetTotalRangeCount(), 1);
+	}
+
+	SECTION("Sorted, non-consecutive indices")
+	{
+		FPCGIndexCollection Collection(20);
+		TArray<int32> Indices = {0, 1, 2, 5, 6, 7, 10, 15};
+		Collection += Indices;
+		REQUIRE_EQUAL(Collection.GetTotalIndexCount(), 8);
+		REQUIRE_EQUAL(Collection.GetTotalRangeCount(), 4);
+		REQUIRE((Collection.ContainsIndex(0) && Collection.ContainsIndex(7) &&
+			Collection.ContainsIndex(10) && Collection.ContainsIndex(15)));
+		REQUIRE_FALSE((Collection.ContainsIndex(3) || Collection.ContainsIndex(8)));
+	}
+
+	SECTION("Unsorted array")
+	{
+		FPCGIndexCollection Collection(20);
+		TArray<int32> Indices = {9, 2, 5, 1, 7, 3};
+		Collection += Indices;
+		REQUIRE_EQUAL(Collection.GetTotalIndexCount(), 6);
+		REQUIRE((Collection.ContainsIndex(1) && Collection.ContainsIndex(9)));
+	}
+
+	SECTION("Duplicates")
+	{
+		FPCGIndexCollection Collection(20);
+		TArray<int32> Indices = {5, 5, 5, 7, 7, 9};
+		Collection += Indices;
+		REQUIRE_EQUAL(Collection.GetTotalIndexCount(), 3);
+		REQUIRE((Collection.ContainsIndex(5) && Collection.ContainsIndex(7) && Collection.ContainsIndex(9)));
+		REQUIRE((!Collection.ContainsIndex(2) && !Collection.ContainsIndex(4) && !Collection.ContainsIndex(6)));
+	}
+
+	SECTION("Negative index")
+	{
+		FPCGIndexCollection Collection(10);
+		Collection += {-1};
+		REQUIRE(Collection.ContainsIndex(9));
+		REQUIRE_EQUAL(Collection.GetTotalIndexCount(), 1);
+	}
+
+	SECTION("Multiple negative indices")
+	{
+		FPCGIndexCollection Collection(10);
+		Collection += {-1, -2};
+		REQUIRE_EQUAL(Collection.GetTotalIndexCount(), 2);
+		REQUIRE_EQUAL(Collection.GetTotalRangeCount(), 1);
+		REQUIRE((Collection.ContainsIndex(8) && Collection.ContainsIndex(9)));
+		REQUIRE_FALSE(Collection.ContainsIndex(7));
+	}
+
+	SECTION("All negative indices spanning full array")
+	{
+		FPCGIndexCollection Collection(5);
+		Collection += {-1, -2, -3, -4, -5};
+		REQUIRE_EQUAL(Collection.GetTotalIndexCount(), 5);
+		REQUIRE_EQUAL(Collection.GetTotalRangeCount(), 1);
+		REQUIRE((Collection.ContainsIndex(0) && Collection.ContainsIndex(4)));
+	}
+
+	SECTION("Negative adjacent to positive after resolution")
+	{
+		FPCGIndexCollection Collection(10);
+		Collection += {8, -1};
+		REQUIRE_EQUAL(Collection.GetTotalIndexCount(), 2);
+		REQUIRE_EQUAL(Collection.GetTotalRangeCount(), 1);
+		REQUIRE((Collection.ContainsIndex(8) && Collection.ContainsIndex(9)));
+	}
+
+	SECTION("Mixed positive and negative indices")
+	{
+		FPCGIndexCollection Collection(10);
+		Collection += {0, 1, -1};
+		REQUIRE_EQUAL(Collection.GetTotalIndexCount(), 3);
+		REQUIRE_EQUAL(Collection.GetTotalRangeCount(), 2);
+		REQUIRE((Collection.ContainsIndex(0) && Collection.ContainsIndex(1) && Collection.ContainsIndex(9)));
+		REQUIRE((!Collection.ContainsIndex(2) && !Collection.ContainsIndex(8)));
+	}
+
+	SECTION("Negative index exceeding array size")
+	{
+		FPCGIndexCollection Collection(10);
+		Collection += {-11};
+		REQUIRE_EQUAL(Collection.GetTotalIndexCount(), 0);
+	}
+
+	SECTION("Array construction with negative indices")
+	{
+		TArray<int32> Indices = {3, -1};
+		FPCGIndexCollection Collection(Indices, 10);
+		REQUIRE_EQUAL(Collection.GetTotalIndexCount(), 2);
+		REQUIRE_EQUAL(Collection.GetTotalRangeCount(), 2);
+		REQUIRE((Collection.ContainsIndex(3) && Collection.ContainsIndex(9)));
+		REQUIRE_FALSE(Collection.ContainsIndex(4));
+	}
+
+	SECTION("Grow size with array with positive index first")
+	{
+		FPCGIndexCollection Collection(10);
+		Collection += {10, -11};
+		// Grow sets ArraySize = max(10 + 1, 10) = 11. Index 10 and 0 are now valid.
+		REQUIRE_EQUAL(Collection.GetArraySize(), 11);
+		REQUIRE_EQUAL(Collection.GetTotalIndexCount(), 2);
+		REQUIRE((Collection.ContainsIndex(0) && Collection.ContainsIndex(10)));
+	}
+
+	SECTION("Grow size with array with negative index first")
+	{
+		FPCGIndexCollection Collection(10);
+		Collection += {-11, 10};
+		// Order shouldn't matter here, so functionally the same as above
+		REQUIRE_EQUAL(Collection.GetArraySize(), 11);
+		REQUIRE_EQUAL(Collection.GetTotalIndexCount(), 2);
+		REQUIRE((Collection.ContainsIndex(0) && Collection.ContainsIndex(10)));
+	}
+}
