[PCG] Fix unit tests log capture after the introduction of `UE_LOGF`

* Serialize is not called anymore, it is SerializeRecord.
* Keep the old for potential backward compatibility, but will be removable when it is deprecated

#rnx

#Synced-CL 51125245

diff --git a/Engine/Plugins/PCG/Tests/Private/PCGTestsCommon.cpp b/Engine/Plugins/PCG/Tests/Private/PCGTestsCommon.cpp
--- a/Engine/Plugins/PCG/Tests/Private/PCGTestsCommon.cpp
+++ b/Engine/Plugins/PCG/Tests/Private/PCGTestsCommon.cpp
@@ -3,6 +3,7 @@
 #include "PCGTestsCommon.h"
 
 #include "PCGModule.h"
+#include "Logging/StructuredLog.h"
 
 namespace PCGTests
 {
@@ -58,4 +59,28 @@ namespace PCGTests
 			OldContext->Serialize(V, Verbosity, Category);
 		}
 	}
+
+	void FPCGTestsLogOutputDevice::SerializeRecord(const UE::FLogRecord& Record)
+	{
+		const ELogVerbosity::Type Verbosity = ResolveVerbosity(Record.GetVerbosity());
+		if (Verbosity <= ELogVerbosity::Warning)
+		{
+			if (NbMessageReceived.Contains(Verbosity))
+			{
+				NbMessageReceived[Verbosity]++;
+			}
+			else
+			{
+				NbMessageReceived.Add(Verbosity, 1);
+			}
+			if (!bSuppressErrors)
+			{
+				OldContext->SerializeRecord(Record);
+			}
+		}
+		else
+		{
+			OldContext->SerializeRecord(Record);
+		}
+	}
 }
diff --git a/Engine/Plugins/PCG/Tests/Private/PCGTestsCommon.h b/Engine/Plugins/PCG/Tests/Private/PCGTestsCommon.h
--- a/Engine/Plugins/PCG/Tests/Private/PCGTestsCommon.h
+++ b/Engine/Plugins/PCG/Tests/Private/PCGTestsCommon.h
@@ -121,6 +121,9 @@ struct FPCGTestsLogOutputDevice : public FFeedbackContext
 	virtual bool IsMemoryOnly() const override { return true; }
 	virtual bool CanBeUsedOnMultipleThreads() const override { return true; }
 	virtual bool CanBeUsedOnAnyThread() const override { return true; }
+	virtual void SerializeRecord(const UE::FLogRecord& Record) override;
+	
+	//~ @todo_pcg: When this becomes deprecated, can be safely removed, as SerializeRecord is the new way of capturing logs.
 	virtual void Serialize(const TCHAR* V, ELogVerbosity::Type Verbosity, const class FName& Category) override;
 
 	bool bSuppressErrors;
