[PCG] Pass on Unit Tests for Metadata

* Convert PCGMetadataTest to Low Level Tests
* Write new Low Level Tests for Attribute
* In preparation for typed attributes becoming Generic Attributes, some tests are already there but deactivated. Annotated with a todo.

#rnx

#Synced-CL 51188099

diff --git a/Engine/Plugins/PCG/Source/PCG/Private/Tests/Metadata/PCGMetadataTest.cpp b/Engine/Plugins/PCG/Source/PCG/Private/Tests/Metadata/PCGMetadataTest.cpp
deleted file mode 100644
diff --git a/Engine/Plugins/PCG/Tests/Private/Metadata/PCGMetadataAttributeTests.cpp b/Engine/Plugins/PCG/Tests/Private/Metadata/PCGMetadataAttributeTests.cpp
new file mode 100644
--- /dev/null
+++ b/Engine/Plugins/PCG/Tests/Private/Metadata/PCGMetadataAttributeTests.cpp
@@ -0,0 +1,1286 @@
+// Copyright Epic Games, Inc. All Rights Reserved.
+
+#include "PCGTestsCommon.h"
+
+#include "PCGParamData.h"
+#include "Data/PCGPointArrayData.h"
+
+#include "Containers/StaticArray.h"
+
+#include <catch2/catch_test_macros.hpp>
+#include <catch2/matchers/catch_matchers_all.hpp>
+
+#include "TestHarness.h"
+
+namespace PCGAttributeTypedTests
+{
+	static const FName AttributeName = "TypedAttr";
+	constexpr int32 Seed = 42;
+	constexpr int32 NumPoints = 100;
+	constexpr int32 NumEntries = 10;
+
+	/**
+	 * Helper to generate a random value of type T using a FRandomStream.
+	 * There are specializations below for each supported type.
+	 */
+	template <typename T>
+	T GenerateRandom(FRandomStream& RandomStream);
+
+	template<> float GenerateRandom<float>(FRandomStream& RandomStream) { return RandomStream.FRand(); }
+	template<> double GenerateRandom<double>(FRandomStream& RandomStream) { return RandomStream.FRand(); }
+	template<> int32 GenerateRandom<int32>(FRandomStream& RandomStream) { return RandomStream.RandRange(-1000, 1000); }
+	template<> int64 GenerateRandom<int64>(FRandomStream& RandomStream) { return static_cast<int64>(RandomStream.RandRange(-100000, 100000)); }
+	template<> bool GenerateRandom<bool>(FRandomStream& RandomStream) { return RandomStream.RandRange(0, 1) == 1; }
+
+	template<> FVector2D GenerateRandom<FVector2D>(FRandomStream& RandomStream)
+	{
+		return FVector2D(RandomStream.FRand(), RandomStream.FRand());
+	}
+
+	template<> FVector GenerateRandom<FVector>(FRandomStream& RandomStream)
+	{
+		return RandomStream.VRand();
+	}
+
+	template<> FVector4 GenerateRandom<FVector4>(FRandomStream& RandomStream)
+	{
+		return FVector4(RandomStream.FRand(), RandomStream.FRand(), RandomStream.FRand(), RandomStream.FRand());
+	}
+
+	template<> FQuat GenerateRandom<FQuat>(FRandomStream& RandomStream)
+	{
+		return FQuat(RandomStream.VRand(), RandomStream.FRand() * UE_PI);
+	}
+
+	template<> FTransform GenerateRandom<FTransform>(FRandomStream& RandomStream)
+	{
+		return FTransform(
+			FQuat(RandomStream.VRand(), RandomStream.FRand() * UE_PI),
+			RandomStream.VRand() * 100.0,
+			FVector(RandomStream.FRand() + 0.5, RandomStream.FRand() + 0.5, RandomStream.FRand() + 0.5));
+	}
+
+	template<> FRotator GenerateRandom<FRotator>(FRandomStream& RandomStream)
+	{
+		return FRotator(RandomStream.FRand() * 360.0 - 180.0, RandomStream.FRand() * 360.0 - 180.0, RandomStream.FRand() * 360.0 - 180.0);
+	}
+
+	template<> FString GenerateRandom<FString>(FRandomStream& RandomStream)
+	{
+		static int Counter = 0;
+		return FString::Printf(TEXT("TestString_%d_%d"), Counter++, RandomStream.RandRange(0, 1000));
+	}
+
+	template<> FName GenerateRandom<FName>(FRandomStream& RandomStream)
+	{
+		static int Counter = 0;
+		return FName(*FString::Printf(TEXT("TestName_%d_%d"), Counter++, RandomStream.RandRange(0, 1000)));
+	}
+
+	template<> FSoftObjectPath GenerateRandom<FSoftObjectPath>(FRandomStream& RandomStream)
+	{
+		static int Counter = 0;
+		return FSoftObjectPath(FString::Printf(TEXT("/Game/Test/Object_%d_%d"), Counter++, RandomStream.RandRange(0, 1000)));
+	}
+
+	template<> FSoftClassPath GenerateRandom<FSoftClassPath>(FRandomStream& RandomStream)
+	{
+		static int Counter = 0;
+		return FSoftClassPath(FString::Printf(TEXT("/Game/Test/Class_%d_%d"), Counter++, RandomStream.RandRange(0, 1000)));
+	}
+
+	/** Default values per type. */
+	template <typename T> T GetDefaultValue();
+	template<> float GetDefaultValue<float>() { return 5.0f; }
+	template<> double GetDefaultValue<double>() { return 5.0; }
+	template<> int32 GetDefaultValue<int32>() { return 42; }
+	template<> int64 GetDefaultValue<int64>() { return 42LL; }
+	template<> bool GetDefaultValue<bool>() { return true; }
+	template<> FVector2D GetDefaultValue<FVector2D>() { return FVector2D(1.0, 2.0); }
+	template<> FVector GetDefaultValue<FVector>() { return FVector(1.0, 2.0, 3.0); }
+	template<> FVector4 GetDefaultValue<FVector4>() { return FVector4(1.0, 2.0, 3.0, 4.0); }
+	template<> FQuat GetDefaultValue<FQuat>() { return FQuat(0.0, 0.0, 0.7071, 0.7071); }
+	template<> FTransform GetDefaultValue<FTransform>() { return FTransform(FQuat::Identity, FVector(5.0, 5.0, 5.0), FVector::OneVector); }
+	template<> FRotator GetDefaultValue<FRotator>() { return FRotator(45.0, 90.0, 0.0); }
+	template<> FString GetDefaultValue<FString>() { return TEXT("DefaultValue"); }
+	template<> FName GetDefaultValue<FName>() { return FName(TEXT("DefaultName")); }
+	template<> FSoftObjectPath GetDefaultValue<FSoftObjectPath>() { return FSoftObjectPath(TEXT("/Game/Default/Object")); }
+	template<> FSoftClassPath GetDefaultValue<FSoftClassPath>() { return FSoftClassPath(TEXT("/Game/Default/Class")); }
+
+	/** Alternative default values for testing default value changes. */
+	template <typename T> T GetAlternativeDefaultValue();
+	template<> float GetAlternativeDefaultValue<float>() { return 99.0f; }
+	template<> double GetAlternativeDefaultValue<double>() { return 99.0; }
+	template<> int32 GetAlternativeDefaultValue<int32>() { return -7; }
+	template<> int64 GetAlternativeDefaultValue<int64>() { return -7LL; }
+	template<> bool GetAlternativeDefaultValue<bool>() { return false; }
+	template<> FVector2D GetAlternativeDefaultValue<FVector2D>() { return FVector2D(99.0, 99.0); }
+	template<> FVector GetAlternativeDefaultValue<FVector>() { return FVector(99.0, 99.0, 99.0); }
+	template<> FVector4 GetAlternativeDefaultValue<FVector4>() { return FVector4(99.0, 99.0, 99.0, 99.0); }
+	template<> FQuat GetAlternativeDefaultValue<FQuat>() { return FQuat::Identity; }
+	template<> FTransform GetAlternativeDefaultValue<FTransform>() { return FTransform::Identity; }
+	template<> FRotator GetAlternativeDefaultValue<FRotator>() { return FRotator::ZeroRotator; }
+	template<> FString GetAlternativeDefaultValue<FString>() { return TEXT("AlternativeDefault"); }
+	template<> FName GetAlternativeDefaultValue<FName>() { return FName(TEXT("AlternativeName")); }
+	template<> FSoftObjectPath GetAlternativeDefaultValue<FSoftObjectPath>() { return FSoftObjectPath(TEXT("/Game/Alternative/Object")); }
+	template<> FSoftClassPath GetAlternativeDefaultValue<FSoftClassPath>() { return FSoftClassPath(TEXT("/Game/Alternative/Class")); }
+
+	/** Expected EPCGMetadataTypes enum for each type. */
+	template <typename T> EPCGMetadataTypes GetExpectedMetadataType();
+	template<> EPCGMetadataTypes GetExpectedMetadataType<float>() { return EPCGMetadataTypes::Float; }
+	template<> EPCGMetadataTypes GetExpectedMetadataType<double>() { return EPCGMetadataTypes::Double; }
+	template<> EPCGMetadataTypes GetExpectedMetadataType<int32>() { return EPCGMetadataTypes::Integer32; }
+	template<> EPCGMetadataTypes GetExpectedMetadataType<int64>() { return EPCGMetadataTypes::Integer64; }
+	template<> EPCGMetadataTypes GetExpectedMetadataType<bool>() { return EPCGMetadataTypes::Boolean; }
+	template<> EPCGMetadataTypes GetExpectedMetadataType<FVector2D>() { return EPCGMetadataTypes::Vector2; }
+	template<> EPCGMetadataTypes GetExpectedMetadataType<FVector>() { return EPCGMetadataTypes::Vector; }
+	template<> EPCGMetadataTypes GetExpectedMetadataType<FVector4>() { return EPCGMetadataTypes::Vector4; }
+	template<> EPCGMetadataTypes GetExpectedMetadataType<FQuat>() { return EPCGMetadataTypes::Quaternion; }
+	template<> EPCGMetadataTypes GetExpectedMetadataType<FTransform>() { return EPCGMetadataTypes::Transform; }
+	template<> EPCGMetadataTypes GetExpectedMetadataType<FRotator>() { return EPCGMetadataTypes::Rotator; }
+	template<> EPCGMetadataTypes GetExpectedMetadataType<FString>() { return EPCGMetadataTypes::String; }
+	template<> EPCGMetadataTypes GetExpectedMetadataType<FName>() { return EPCGMetadataTypes::Name; }
+	template<> EPCGMetadataTypes GetExpectedMetadataType<FSoftObjectPath>() { return EPCGMetadataTypes::SoftObjectPath; }
+	template<> EPCGMetadataTypes GetExpectedMetadataType<FSoftClassPath>() { return EPCGMetadataTypes::SoftClassPath; }
+	
+	// @todo_pcg To be replaced by `IsOfType<T>` when generic attributes are the base of the typed one.
+	template <typename T>
+	bool IsOfType(const FPCGMetadataAttributeBase* Attribute)
+	{
+		return Attribute && PCG::Private::MetadataTypes<T>::Id == Attribute->GetTypeId();
+	}
+	
+}
+
+/**
+ * Test creating FPCGMetadataAttribute<T> via Domain->CreateAttribute<T> for all supported types,
+ * verify type checking, default values, single and bulk set/get, and parenting with PointArrayData.
+ *
+ * The typed tester lambda operates generically over T:
+ *   1. Create a UPCGPointArrayData, allocate metadata entries.
+ *   2. Create an FPCGMetadataAttribute<T> via CreateAttribute.
+ *   3. Verify the attribute is of the right type (IsOfType<T>).
+ *   4. Verify the default value can be read back correctly.
+ *   5. Set random values in bulk, read them back and verify.
+ *   6. Test parenting: create a child data, inherit, modify even indices, verify all values.
+ */
+TEST_CASE_METHOD(PCGTests::FPCGBaseTestWithBasicContext, "PCG::AttributeTyped::BasicWithPointArrayData", "[PCG][AttributeTyped]")
+{
+	using namespace PCGAttributeTypedTests;
+
+	auto TypedTester = [this]<typename T>()
+	{
+		FRandomStream RandomStream(Seed);
+
+		// Create the data with points
+		UPCGPointArrayData* PointArrayData = CreateData<UPCGPointArrayData>();
+		PointArrayData->SetNumPoints(NumPoints);
+		PointArrayData->AllocateProperties(EPCGPointNativeProperties::MetadataEntry);
+		TPCGValueRange<int64> MetadataEntryRange = PointArrayData->GetMetadataEntryValueRange(/*bAllocate=*/false);
+
+		FPCGMetadataDomain* Domain = PointArrayData->MutableMetadata()->GetDefaultMetadataDomain();
+		REQUIRE_NOT_EQUAL(Domain, nullptr);
+
+		// Create the typed attribute
+		const T DefaultValue = GetDefaultValue<T>();
+		FPCGMetadataAttribute<T>* TypedAttribute = Domain->CreateAttribute<T>(AttributeName, DefaultValue, /*bAllowsInterpolation=*/true, /*bOverrideParent=*/false);
+		REQUIRE_NOT_EQUAL(TypedAttribute, nullptr);
+		
+		SECTION("IsOfType check")
+		{
+			REQUIRE(IsOfType<T>(TypedAttribute));
+		}
+
+		// @todo_pcg to be enabled when typed attribute inherit from generic
+		// SECTION("Attribute descriptor matches expected type")
+		// {
+		// 	const FPCGMetadataAttributeDesc& Desc = TypedAttribute->GetAttributeDesc();
+		// 	REQUIRE_EQUAL(Desc.Name, AttributeName);
+		// 	REQUIRE_EQUAL(Desc.ValueType, GetExpectedMetadataType<T>());
+		// }
+
+		SECTION("Default value is set correctly")
+		{
+			T ReadDefaultValue = TypedAttribute->GetValueFromItemKey(PCGInvalidEntryKey);
+			if constexpr (std::is_same_v<T, FTransform>)
+			{
+				REQUIRE(ReadDefaultValue.Equals(DefaultValue));
+			}
+			else
+			{
+				REQUIRE_EQUAL(ReadDefaultValue, DefaultValue);
+			}
+		}
+
+		// Add entries and set random values
+		TStaticArray<T, NumPoints> Values{};
+		for (int32 i = 0; i < NumPoints; ++i)
+		{
+			MetadataEntryRange[i] = Domain->AddEntry(MetadataEntryRange[i]);
+			Values[i] = GenerateRandom<T>(RandomStream);
+		}
+
+		// Bulk set
+		TypedAttribute->SetValues(PCGValueRangeHelpers::MakeConstValueRange(MetadataEntryRange), Values);
+		REQUIRE_EQUAL(TypedAttribute->GetNumberOfEntries(), NumPoints);
+
+		SECTION("Read back and verify all values")
+		{
+			TStaticArray<T, NumPoints> ReadValues{};
+			TypedAttribute->GetValuesFromItemKeys(PCGValueRangeHelpers::MakeConstValueRange(MetadataEntryRange), ReadValues);
+
+			for (int32 i = 0; i < NumPoints; ++i)
+			{
+				if constexpr (std::is_same_v<T, FTransform>)
+				{
+					REQUIRE(ReadValues[i].Equals(Values[i]));
+				}
+				else
+				{
+					REQUIRE_EQUAL(ReadValues[i], Values[i]);
+				}
+			}
+		}
+
+		SECTION("Single value set and get")
+		{
+			// Override the first entry with a specific value
+			T SpecificValue = GetAlternativeDefaultValue<T>();
+			TypedAttribute->SetValue(MetadataEntryRange[0], SpecificValue);
+
+			T ReadBack = TypedAttribute->GetValueFromItemKey(MetadataEntryRange[0]);
+			if constexpr (std::is_same_v<T, FTransform>)
+			{
+				REQUIRE(ReadBack.Equals(SpecificValue));
+			}
+			else
+			{
+				REQUIRE_EQUAL(ReadBack, SpecificValue);
+			}
+		}
+
+		SECTION("Parenting with child PointArrayData")
+		{
+			UPCGPointArrayData* ChildPointArrayData = CreateData<UPCGPointArrayData>();
+			ChildPointArrayData->InitializeFromDataWithParams(FPCGInitializeFromDataParams{PointArrayData});
+
+			FPCGMetadataDomain* ChildDomain = ChildPointArrayData->MutableMetadata()->GetDefaultMetadataDomain();
+			REQUIRE_NOT_EQUAL(ChildDomain, nullptr);
+
+			// Retrieve the inherited typed attribute
+			FPCGMetadataAttribute<T>* ChildTypedAttribute = ChildDomain->GetMutableTypedAttribute<T>(AttributeName);
+			REQUIRE_NOT_EQUAL(ChildTypedAttribute, nullptr);
+
+			TPCGValueRange<int64> ChildMetadataEntryRange = ChildPointArrayData->GetMetadataEntryValueRange(/*bAllocate=*/false);
+
+			SECTION("Child inherits default value")
+			{
+				T ReadDefault = ChildTypedAttribute->GetValueFromItemKey(PCGInvalidEntryKey);
+				if constexpr (std::is_same_v<T, FTransform>)
+				{
+					REQUIRE(ReadDefault.Equals(DefaultValue));
+				}
+				else
+				{
+					REQUIRE_EQUAL(ReadDefault, DefaultValue);
+				}
+			}
+
+			SECTION("Child reads parent values unchanged")
+			{
+				TStaticArray<T, NumPoints> ReadValues{};
+				ChildTypedAttribute->GetValuesFromItemKeys(PCGValueRangeHelpers::MakeConstValueRange(ChildMetadataEntryRange), ReadValues);
+
+				for (int32 i = 0; i < NumPoints; ++i)
+				{
+					if constexpr (std::is_same_v<T, FTransform>)
+					{
+						REQUIRE(ReadValues[i].Equals(Values[i]));
+					}
+					else
+					{
+						REQUIRE_EQUAL(ReadValues[i], Values[i]);
+					}
+				}
+			}
+
+			SECTION("Child modifies even indices, odd indices read from parent")
+			{
+				// Modify even indices
+				TStaticArray<T, NumPoints / 2> ChildValues{};
+				TStaticArray<PCGMetadataEntryKey, NumPoints / 2> ChildEntryKeys{};
+				for (int32 i = 0; i < NumPoints; ++i)
+				{
+					if (i % 2 == 0)
+					{
+						Values[i] = GenerateRandom<T>(RandomStream);
+						ChildDomain->InitializeOnSet(ChildMetadataEntryRange[i]);
+						ChildValues[i / 2] = Values[i];
+						ChildEntryKeys[i / 2] = ChildMetadataEntryRange[i];
+					}
+				}
+
+				ChildTypedAttribute->SetValues(TConstArrayView<PCGMetadataEntryKey>(ChildEntryKeys), ChildValues);
+
+				REQUIRE_EQUAL(ChildTypedAttribute->GetNumberOfEntries(), NumPoints / 2);
+				REQUIRE_EQUAL(ChildTypedAttribute->GetNumberOfEntriesWithParents(), NumPoints + NumPoints / 2);
+
+				// Read all values - even indices should have new values, odd indices should have parent values
+				TStaticArray<T, NumPoints> ReadValues{};
+				ChildTypedAttribute->GetValuesFromItemKeys(PCGValueRangeHelpers::MakeConstValueRange(ChildMetadataEntryRange), ReadValues);
+
+				for (int32 i = 0; i < NumPoints; ++i)
+				{
+					if constexpr (std::is_same_v<T, FTransform>)
+					{
+						REQUIRE(ReadValues[i].Equals(Values[i]));
+					}
+					else
+					{
+						REQUIRE_EQUAL(ReadValues[i], Values[i]);
+					}
+				}
+			}
+		}
+	};
+
+	SECTION("float") { TypedTester.template operator()<float>(); }
+	SECTION("double") { TypedTester.template operator()<double>(); }
+	SECTION("int32") { TypedTester.template operator()<int32>(); }
+	SECTION("int64") { TypedTester.template operator()<int64>(); }
+	SECTION("bool") { TypedTester.template operator()<bool>(); }
+	SECTION("FVector2D") { TypedTester.template operator()<FVector2D>(); }
+	SECTION("FVector") { TypedTester.template operator()<FVector>(); }
+	SECTION("FVector4") { TypedTester.template operator()<FVector4>(); }
+	SECTION("FQuat") { TypedTester.template operator()<FQuat>(); }
+	SECTION("FTransform") { TypedTester.template operator()<FTransform>(); }
+	SECTION("FRotator") { TypedTester.template operator()<FRotator>(); }
+	SECTION("FString") { TypedTester.template operator()<FString>(); }
+	SECTION("FName") { TypedTester.template operator()<FName>(); }
+	SECTION("FSoftObjectPath") { TypedTester.template operator()<FSoftObjectPath>(); }
+	SECTION("FSoftClassPath") { TypedTester.template operator()<FSoftClassPath>(); }
+}
+
+/**
+ * Test creating FPCGMetadataAttribute<T> via Domain->CreateAttribute<T> with ParamData (copy-based parenting).
+ * Tests that inherited values through copy are preserved correctly.
+ */
+TEST_CASE_METHOD(PCGTests::FPCGBaseTestWithBasicContext, "PCG::AttributeTyped::CopyWithParamData", "[PCG][AttributeTyped]")
+{
+	using namespace PCGAttributeTypedTests;
+
+	auto TypedTester = [this]<typename T>()
+	{
+		FRandomStream RandomStream(Seed);
+
+		UPCGParamData* ParamData = CreateData<UPCGParamData>();
+		FPCGMetadataDomain* Domain = ParamData->MutableMetadata()->GetDefaultMetadataDomain();
+		REQUIRE_NOT_EQUAL(Domain, nullptr);
+
+		const T DefaultValue = GetDefaultValue<T>();
+		FPCGMetadataAttribute<T>* TypedAttribute = Domain->CreateAttribute<T>(AttributeName, DefaultValue, /*bAllowsInterpolation=*/true, /*bOverrideParent=*/false);
+		REQUIRE_NOT_EQUAL(TypedAttribute, nullptr);
+
+		TStaticArray<T, NumEntries> Values{};
+		TStaticArray<PCGMetadataEntryKey, NumEntries> EntryKeys;
+
+		for (int32 i = 0; i < NumEntries; ++i)
+		{
+			EntryKeys[i] = Domain->AddEntry();
+			Values[i] = GenerateRandom<T>(RandomStream);
+		}
+
+		TypedAttribute->SetValues(TConstArrayView<PCGMetadataEntryKey>(EntryKeys), Values);
+		REQUIRE_EQUAL(TypedAttribute->GetNumberOfEntries(), NumEntries);
+
+		SECTION("Values read back correctly")
+		{
+			TStaticArray<T, NumEntries> ReadValues{};
+			TypedAttribute->GetValuesFromItemKeys(TConstArrayView<PCGMetadataEntryKey>(EntryKeys), ReadValues);
+
+			for (int32 i = 0; i < NumEntries; ++i)
+			{
+				if constexpr (std::is_same_v<T, FTransform>)
+				{
+					REQUIRE(ReadValues[i].Equals(Values[i]));
+				}
+				else
+				{
+					REQUIRE_EQUAL(ReadValues[i], Values[i]);
+				}
+			}
+		}
+
+		SECTION("Copy-based parenting")
+		{
+			UPCGParamData* ChildParamData = CreateData<UPCGParamData>();
+			ChildParamData->MutableMetadata()->Initialize(ParamData->ConstMetadata());
+
+			FPCGMetadataDomain* ChildDomain = ChildParamData->MutableMetadata()->GetDefaultMetadataDomain();
+			REQUIRE_NOT_EQUAL(ChildDomain, nullptr);
+
+			FPCGMetadataAttribute<T>* ChildTypedAttribute = ChildDomain->GetMutableTypedAttribute<T>(AttributeName);
+			REQUIRE_NOT_EQUAL(ChildTypedAttribute, nullptr);
+
+			SECTION("Inherited default value")
+			{
+				T ReadDefault = ChildTypedAttribute->GetValueFromItemKey(PCGInvalidEntryKey);
+				if constexpr (std::is_same_v<T, FTransform>)
+				{
+					REQUIRE(ReadDefault.Equals(DefaultValue));
+				}
+				else
+				{
+					REQUIRE_EQUAL(ReadDefault, DefaultValue);
+				}
+			}
+
+			SECTION("Modify even entries, verify all")
+			{
+				TStaticArray<T, NumEntries / 2> ChildValues{};
+				TStaticArray<PCGMetadataEntryKey, NumEntries / 2> ChildEntryKeys{};
+				for (int32 i = 0; i < NumEntries; ++i)
+				{
+					if (i % 2 == 0)
+					{
+						Values[i] = GenerateRandom<T>(RandomStream);
+						ChildDomain->InitializeOnSet(EntryKeys[i]);
+						ChildValues[i / 2] = Values[i];
+						ChildEntryKeys[i / 2] = EntryKeys[i];
+					}
+				}
+
+				ChildTypedAttribute->SetValues(TConstArrayView<PCGMetadataEntryKey>(ChildEntryKeys), ChildValues);
+
+				// For copy-based parenting, entries are copied, so total should be NumEntries (not additive)
+				REQUIRE_EQUAL(ChildTypedAttribute->GetNumberOfEntries(), NumEntries);
+				REQUIRE_EQUAL(ChildTypedAttribute->GetNumberOfEntriesWithParents(), NumEntries);
+
+				TStaticArray<T, NumEntries> ReadValues{};
+				ChildTypedAttribute->GetValuesFromItemKeys(TConstArrayView<PCGMetadataEntryKey>(EntryKeys), ReadValues);
+
+				for (int32 i = 0; i < NumEntries; ++i)
+				{
+					if constexpr (std::is_same_v<T, FTransform>)
+					{
+						REQUIRE(ReadValues[i].Equals(Values[i]));
+					}
+					else
+					{
+						REQUIRE_EQUAL(ReadValues[i], Values[i]);
+					}
+				}
+			}
+		}
+	};
+
+	SECTION("float") { TypedTester.template operator()<float>(); }
+	SECTION("double") { TypedTester.template operator()<double>(); }
+	SECTION("int32") { TypedTester.template operator()<int32>(); }
+	SECTION("int64") { TypedTester.template operator()<int64>(); }
+	SECTION("bool") { TypedTester.template operator()<bool>(); }
+	SECTION("FVector2D") { TypedTester.template operator()<FVector2D>(); }
+	SECTION("FVector") { TypedTester.template operator()<FVector>(); }
+	SECTION("FVector4") { TypedTester.template operator()<FVector4>(); }
+	SECTION("FQuat") { TypedTester.template operator()<FQuat>(); }
+	SECTION("FTransform") { TypedTester.template operator()<FTransform>(); }
+	SECTION("FRotator") { TypedTester.template operator()<FRotator>(); }
+	SECTION("FString") { TypedTester.template operator()<FString>(); }
+	SECTION("FName") { TypedTester.template operator()<FName>(); }
+	SECTION("FSoftObjectPath") { TypedTester.template operator()<FSoftObjectPath>(); }
+	SECTION("FSoftClassPath") { TypedTester.template operator()<FSoftClassPath>(); }
+}
+
+/**
+ * Test FindOrCreateAttribute: first call creates, second call finds the existing one.
+ * Also test that FindOrCreateAttribute with a different type and bOverwriteIfTypeMismatch=true
+ * replaces the attribute.
+ */
+TEST_CASE_METHOD(PCGTests::FPCGBaseTestWithBasicContext, "PCG::AttributeTyped::FindOrCreate", "[PCG][AttributeTyped]")
+{
+	using namespace PCGAttributeTypedTests;
+
+	UPCGParamData* ParamData = CreateData<UPCGParamData>();
+	FPCGMetadataDomain* Domain = ParamData->MutableMetadata()->GetDefaultMetadataDomain();
+	REQUIRE_NOT_EQUAL(Domain, nullptr);
+
+	SECTION("Create then find returns same attribute")
+	{
+		FPCGMetadataAttribute<float>* Attr1 = Domain->FindOrCreateAttribute<float>(AttributeName, 1.0f);
+		REQUIRE_NOT_EQUAL(Attr1, nullptr);
+
+		FPCGMetadataAttribute<float>* Attr2 = Domain->FindOrCreateAttribute<float>(AttributeName, 2.0f);
+		REQUIRE_NOT_EQUAL(Attr2, nullptr);
+
+		// Should be the same attribute pointer
+		REQUIRE_EQUAL(Attr1, Attr2);
+
+		// Default value should still be from the first creation
+		float DefaultVal = Attr2->GetValueFromItemKey(PCGInvalidEntryKey);
+		REQUIRE_EQUAL(DefaultVal, 1.0f);
+	}
+
+	SECTION("FindOrCreate with type mismatch replaces attribute")
+	{
+		FPCGMetadataAttribute<float>* FloatAttr = Domain->FindOrCreateAttribute<float>(AttributeName, 1.0f);
+		REQUIRE_NOT_EQUAL(FloatAttr, nullptr);
+		REQUIRE(IsOfType<float>(FloatAttr));
+
+		// Now request an int32 attribute with the same name - should overwrite
+		FPCGMetadataAttribute<int32>* IntAttr = Domain->FindOrCreateAttribute<int32>(AttributeName, 42, /*bAllowsInterpolation=*/true, /*bOverrideParent=*/true, /*bOverwriteIfTypeMismatch=*/true);
+		REQUIRE_NOT_EQUAL(IntAttr, nullptr);
+		REQUIRE(IsOfType<int32>(IntAttr));
+
+		int32 DefaultVal = IntAttr->GetValueFromItemKey(PCGInvalidEntryKey);
+		REQUIRE_EQUAL(DefaultVal, 42);
+	}
+}
+
+/**
+ * Test that GetValue and GetValues using value keys (not entry keys) work correctly.
+ */
+TEST_CASE_METHOD(PCGTests::FPCGBaseTestWithBasicContext, "PCG::AttributeTyped::GetValueByValueKey", "[PCG][AttributeTyped]")
+{
+	using namespace PCGAttributeTypedTests;
+
+	auto TypedTester = [this]<typename T>()
+	{
+		UPCGParamData* ParamData = CreateData<UPCGParamData>();
+		FPCGMetadataDomain* Domain = ParamData->MutableMetadata()->GetDefaultMetadataDomain();
+
+		const T DefaultValue = GetDefaultValue<T>();
+		FPCGMetadataAttribute<T>* Attr = Domain->CreateAttribute<T>(AttributeName, DefaultValue, /*bAllowsInterpolation=*/true, /*bOverrideParent=*/false);
+		REQUIRE_NOT_EQUAL(Attr, nullptr);
+
+		// GetValue with PCGDefaultValueKey should return the default
+		T ReadDefault = Attr->GetValue(PCGDefaultValueKey);
+		if constexpr (std::is_same_v<T, FTransform>)
+		{
+			REQUIRE(ReadDefault.Equals(DefaultValue));
+		}
+		else
+		{
+			REQUIRE_EQUAL(ReadDefault, DefaultValue);
+		}
+
+		// Add a value and read it by value key
+		FRandomStream RandomStream(Seed);
+		T TestValue = GenerateRandom<T>(RandomStream);
+		PCGMetadataValueKey ValueKey = Attr->AddValue(TestValue);
+		REQUIRE(ValueKey != PCGDefaultValueKey);
+
+		T ReadValue = Attr->GetValue(ValueKey);
+		if constexpr (std::is_same_v<T, FTransform>)
+		{
+			REQUIRE(ReadValue.Equals(TestValue));
+		}
+		else
+		{
+			REQUIRE_EQUAL(ReadValue, TestValue);
+		}
+	};
+
+	SECTION("float") { TypedTester.template operator()<float>(); }
+	SECTION("double") { TypedTester.template operator()<double>(); }
+	SECTION("int32") { TypedTester.template operator()<int32>(); }
+	SECTION("int64") { TypedTester.template operator()<int64>(); }
+	SECTION("bool") { TypedTester.template operator()<bool>(); }
+	SECTION("FVector2D") { TypedTester.template operator()<FVector2D>(); }
+	SECTION("FVector") { TypedTester.template operator()<FVector>(); }
+	SECTION("FVector4") { TypedTester.template operator()<FVector4>(); }
+	SECTION("FQuat") { TypedTester.template operator()<FQuat>(); }
+	SECTION("FTransform") { TypedTester.template operator()<FTransform>(); }
+	SECTION("FRotator") { TypedTester.template operator()<FRotator>(); }
+	SECTION("FString") { TypedTester.template operator()<FString>(); }
+	SECTION("FName") { TypedTester.template operator()<FName>(); }
+	SECTION("FSoftObjectPath") { TypedTester.template operator()<FSoftObjectPath>(); }
+	SECTION("FSoftClassPath") { TypedTester.template operator()<FSoftClassPath>(); }
+}
+
+/**
+ * Test AddValue and AddValues return correct value keys, and that FindValue works
+ * for compressed data types (FString, FSoftObjectPath, FSoftClassPath).
+ */
+TEST_CASE_METHOD(PCGTests::FPCGBaseTestWithBasicContext, "PCG::AttributeTyped::AddValueAndFindValue", "[PCG][AttributeTyped]")
+{
+	using namespace PCGAttributeTypedTests;
+
+	UPCGParamData* ParamData = CreateData<UPCGParamData>();
+	FPCGMetadataDomain* Domain = ParamData->MutableMetadata()->GetDefaultMetadataDomain();
+
+	SECTION("FString - compressed, FindValue should work")
+	{
+		FPCGMetadataAttribute<FString>* Attr = Domain->CreateAttribute<FString>(AttributeName, FString(TEXT("Default")), /*bAllowsInterpolation=*/true, /*bOverrideParent=*/false);
+		REQUIRE_NOT_EQUAL(Attr, nullptr);
+
+		PCGMetadataValueKey Key1 = Attr->AddValue(FString(TEXT("Hello")));
+		PCGMetadataValueKey Key2 = Attr->AddValue(FString(TEXT("World")));
+		REQUIRE(Key1 != PCGDefaultValueKey);
+		REQUIRE(Key2 != PCGDefaultValueKey);
+		REQUIRE(Key1 != Key2);
+
+		// Adding the same value should return the same key (compressed)
+		PCGMetadataValueKey Key1Again = Attr->AddValue(FString(TEXT("Hello")));
+		REQUIRE_EQUAL(Key1, Key1Again);
+
+		// FindValue should locate existing values
+		PCGMetadataValueKey FoundKey = Attr->FindValue(FString(TEXT("Hello")));
+		REQUIRE_EQUAL(FoundKey, Key1);
+
+		PCGMetadataValueKey FoundKey2 = Attr->FindValue(FString(TEXT("World")));
+		REQUIRE_EQUAL(FoundKey2, Key2);
+	}
+
+	SECTION("FSoftObjectPath - compressed, FindValue should work")
+	{
+		FPCGMetadataAttribute<FSoftObjectPath>* Attr = Domain->CreateAttribute<FSoftObjectPath>(AttributeName, FSoftObjectPath(), /*bAllowsInterpolation=*/true, /*bOverrideParent=*/false);
+		REQUIRE_NOT_EQUAL(Attr, nullptr);
+
+		FSoftObjectPath Path1(TEXT("/Game/Test/A"));
+		FSoftObjectPath Path2(TEXT("/Game/Test/B"));
+
+		PCGMetadataValueKey Key1 = Attr->AddValue(Path1);
+		PCGMetadataValueKey Key2 = Attr->AddValue(Path2);
+		REQUIRE(Key1 != Key2);
+
+		PCGMetadataValueKey FoundKey = Attr->FindValue(Path1);
+		REQUIRE_EQUAL(FoundKey, Key1);
+	}
+
+	SECTION("FSoftClassPath - compressed, FindValue should work")
+	{
+		FPCGMetadataAttribute<FSoftClassPath>* Attr = Domain->CreateAttribute<FSoftClassPath>(AttributeName, FSoftClassPath(), /*bAllowsInterpolation=*/true, /*bOverrideParent=*/false);
+		REQUIRE_NOT_EQUAL(Attr, nullptr);
+
+		FSoftClassPath Path1(TEXT("/Game/Test/ClassA"));
+		FSoftClassPath Path2(TEXT("/Game/Test/ClassB"));
+
+		PCGMetadataValueKey Key1 = Attr->AddValue(Path1);
+		PCGMetadataValueKey Key2 = Attr->AddValue(Path2);
+		REQUIRE(Key1 != Key2);
+
+		PCGMetadataValueKey FoundKey = Attr->FindValue(Path1);
+		REQUIRE_EQUAL(FoundKey, Key1);
+	}
+
+	SECTION("float - not compressed, FindValue returns PCGNotFoundValueKey")
+	{
+		FPCGMetadataAttribute<float>* Attr = Domain->CreateAttribute<float>(AttributeName, 0.0f, /*bAllowsInterpolation=*/true, /*bOverrideParent=*/false);
+		REQUIRE_NOT_EQUAL(Attr, nullptr);
+
+		Attr->AddValue(1.0f);
+		PCGMetadataValueKey FoundKey = Attr->FindValue(1.0f);
+		REQUIRE_EQUAL(FoundKey, PCGNotFoundValueKey);
+	}
+}
+
+/**
+ * Test that the typed attribute is accessible through both the generic and typed getter APIs.
+ */
+TEST_CASE_METHOD(PCGTests::FPCGBaseTestWithBasicContext, "PCG::AttributeTyped::GenericAndTypedAccess", "[PCG][AttributeTyped]")
+{
+	using namespace PCGAttributeTypedTests;
+
+	UPCGParamData* ParamData = CreateData<UPCGParamData>();
+	FPCGMetadataDomain* Domain = ParamData->MutableMetadata()->GetDefaultMetadataDomain();
+
+	FPCGMetadataAttribute<double>* TypedAttr = Domain->CreateAttribute<double>(AttributeName, 3.14, /*bAllowsInterpolation=*/true, /*bOverrideParent=*/false);
+	REQUIRE_NOT_EQUAL(TypedAttr, nullptr);
+
+	// @todo_pcg to be enabled when typed attribute inherit from generic
+	// SECTION("GetMutableGenericAttribute returns the same attribute")
+	// {
+	// 	FPCGMetadataAttributeGeneric* GenericAttr = Domain->GetMutableGenericAttribute(AttributeName);
+	// 	REQUIRE_NOT_EQUAL(GenericAttr, nullptr);
+	// 	REQUIRE(IsOfType<double>(GenericAttr));
+	//
+	// 	REQUIRE_EQUAL(static_cast<FPCGMetadataAttributeGeneric*>(TypedAttr), GenericAttr);
+	// }
+
+	SECTION("GetMutableTypedAttribute returns typed attribute")
+	{
+		FPCGMetadataAttribute<double>* FoundAttr = Domain->GetMutableTypedAttribute<double>(AttributeName);
+		REQUIRE_EQUAL(FoundAttr, TypedAttr);
+	}
+
+	SECTION("GetMutableTypedAttribute with wrong type returns nullptr")
+	{
+		FPCGMetadataAttribute<float>* WrongAttr = Domain->GetMutableTypedAttribute<float>(AttributeName);
+		REQUIRE_EQUAL(WrongAttr, nullptr);
+	}
+
+
+	// @todo_pcg to be enabled when typed attribute inherit from generic
+	//SECTION("Setting via generic, reading via typed")
+	//{
+	//	PCGMetadataEntryKey EntryKey = Domain->AddEntry();
+	//	FPCGMetadataAttributeGeneric* GenericAttr = Domain->GetMutableGenericAttribute(AttributeName);
+	//	GenericAttr->SetValue<double>(EntryKey, 2.718);
+
+	//	double ReadValue = TypedAttr->GetValueFromItemKey(EntryKey);
+	//	REQUIRE_EQUAL(ReadValue, 2.718);
+	//}
+
+	// SECTION("Setting via typed, reading via generic")
+	// {
+	// 	PCGMetadataEntryKey EntryKey = Domain->AddEntry();
+	// 	TypedAttr->SetValue(EntryKey, 1.618);
+	//
+	// 	FPCGMetadataAttributeGeneric* GenericAttr = Domain->GetMutableGenericAttribute(AttributeName);
+	// 	double ReadValue = GenericAttr->GetValueFromItemKey<double>(EntryKey);
+	// 	REQUIRE_EQUAL(ReadValue, 1.618);
+	// }
+}
+
+/**
+ * Test TypedCopy: creating a copy of a typed attribute with a new name,
+ * verifying that entries and values are preserved correctly.
+ */
+TEST_CASE_METHOD(PCGTests::FPCGBaseTestWithBasicContext, "PCG::AttributeTyped::TypedCopy", "[PCG][AttributeTyped]")
+{
+	using namespace PCGAttributeTypedTests;
+
+	auto TypedTester = [this]<typename T>()
+	{
+		FRandomStream RandomStream(Seed);
+
+		UPCGParamData* ParamData = CreateData<UPCGParamData>();
+		FPCGMetadataDomain* Domain = ParamData->MutableMetadata()->GetDefaultMetadataDomain();
+
+		const T DefaultValue = GetDefaultValue<T>();
+		FPCGMetadataAttribute<T>* OriginalAttr = Domain->CreateAttribute<T>(AttributeName, DefaultValue, /*bAllowsInterpolation=*/true, /*bOverrideParent=*/false);
+		REQUIRE_NOT_EQUAL(OriginalAttr, nullptr);
+
+		constexpr int32 CopyCount = 5;
+		TStaticArray<T, CopyCount> Values{};
+		TStaticArray<PCGMetadataEntryKey, CopyCount> EntryKeys;
+
+		for (int32 i = 0; i < CopyCount; ++i)
+		{
+			EntryKeys[i] = Domain->AddEntry();
+			Values[i] = GenerateRandom<T>(RandomStream);
+		}
+
+		OriginalAttr->SetValues(TConstArrayView<PCGMetadataEntryKey>(EntryKeys), Values);
+
+		FPCGMetadataAttribute<T>* CopiedAttr = OriginalAttr->TypedCopy(FName(TEXT("CopiedAttr")), Domain, /*bKeepParent=*/false, /*bCopyEntries=*/true, /*bCopyValues=*/true);
+		REQUIRE_NOT_EQUAL(CopiedAttr, nullptr);
+
+		SECTION("Copied attribute has correct default")
+		{
+			T ReadDefault = CopiedAttr->GetValue(PCGDefaultValueKey);
+			REQUIRE_EQUAL(ReadDefault, DefaultValue);
+		}
+
+		SECTION("Copied attribute has correct values")
+		{
+			TStaticArray<T, CopyCount> ReadValues{};
+			CopiedAttr->GetValuesFromItemKeys(TConstArrayView<PCGMetadataEntryKey>(EntryKeys), ReadValues);
+
+			for (int32 i = 0; i < CopyCount; ++i)
+			{
+				REQUIRE_EQUAL(ReadValues[i], Values[i]);
+			}
+		}
+
+		// Clean up the copied attribute as it was created with new
+		delete CopiedAttr;
+	};
+
+	SECTION("float") { TypedTester.template operator()<float>(); }
+	SECTION("double") { TypedTester.template operator()<double>(); }
+	SECTION("int32") { TypedTester.template operator()<int32>(); }
+	SECTION("int64") { TypedTester.template operator()<int64>(); }
+	SECTION("FVector") { TypedTester.template operator()<FVector>(); }
+	SECTION("FString") { TypedTester.template operator()<FString>(); }
+	SECTION("FName") { TypedTester.template operator()<FName>(); }
+}
+
+/**
+ * Edge case: creating an attribute with zero entries and reading the default value.
+ */
+TEST_CASE_METHOD(PCGTests::FPCGBaseTestWithBasicContext, "PCG::AttributeTyped::EmptyAttribute", "[PCG][AttributeTyped]")
+{
+	using namespace PCGAttributeTypedTests;
+
+	UPCGParamData* ParamData = CreateData<UPCGParamData>();
+	FPCGMetadataDomain* Domain = ParamData->MutableMetadata()->GetDefaultMetadataDomain();
+
+	SECTION("No entries, read default for float")
+	{
+		FPCGMetadataAttribute<float>* Attr = Domain->CreateAttribute<float>(AttributeName, 42.0f, /*bAllowsInterpolation=*/true, /*bOverrideParent=*/false);
+		REQUIRE_NOT_EQUAL(Attr, nullptr);
+		REQUIRE_EQUAL(Attr->GetNumberOfEntries(), 0);
+
+		float ReadDefault = Attr->GetValueFromItemKey(PCGInvalidEntryKey);
+		REQUIRE_EQUAL(ReadDefault, 42.0f);
+	}
+
+	SECTION("No entries, read default for FString")
+	{
+		FPCGMetadataAttribute<FString>* Attr = Domain->CreateAttribute<FString>(AttributeName, FString(TEXT("Empty")), /*bAllowsInterpolation=*/true, /*bOverrideParent=*/false);
+		REQUIRE_NOT_EQUAL(Attr, nullptr);
+		REQUIRE_EQUAL(Attr->GetNumberOfEntries(), 0);
+
+		FString ReadDefault = Attr->GetValueFromItemKey(PCGInvalidEntryKey);
+		REQUIRE_EQUAL(ReadDefault, FString(TEXT("Empty")));
+	}
+
+	SECTION("No entries, read default for bool")
+	{
+		FPCGMetadataAttribute<bool>* Attr = Domain->CreateAttribute<bool>(AttributeName, true, /*bAllowsInterpolation=*/true, /*bOverrideParent=*/false);
+		REQUIRE_NOT_EQUAL(Attr, nullptr);
+
+		bool ReadDefault = Attr->GetValueFromItemKey(PCGInvalidEntryKey);
+		REQUIRE_EQUAL(ReadDefault, true);
+	}
+}
+
+/**
+ * Test overwriting the same entry key multiple times and verifying only the last value persists.
+ */
+TEST_CASE_METHOD(PCGTests::FPCGBaseTestWithBasicContext, "PCG::AttributeTyped::OverwriteValues", "[PCG][AttributeTyped]")
+{
+	using namespace PCGAttributeTypedTests;
+
+	UPCGParamData* ParamData = CreateData<UPCGParamData>();
+	FPCGMetadataDomain* Domain = ParamData->MutableMetadata()->GetDefaultMetadataDomain();
+
+	FPCGMetadataAttribute<int32>* Attr = Domain->CreateAttribute<int32>(AttributeName, 0, /*bAllowsInterpolation=*/true, /*bOverrideParent=*/false);
+	REQUIRE_NOT_EQUAL(Attr, nullptr);
+
+	PCGMetadataEntryKey EntryKey = Domain->AddEntry();
+
+	// Write multiple times to the same entry
+	Attr->SetValue(EntryKey, 10);
+	REQUIRE_EQUAL(Attr->GetValueFromItemKey(EntryKey), 10);
+
+	Attr->SetValue(EntryKey, 20);
+	REQUIRE_EQUAL(Attr->GetValueFromItemKey(EntryKey), 20);
+
+	Attr->SetValue(EntryKey, 30);
+	REQUIRE_EQUAL(Attr->GetValueFromItemKey(EntryKey), 30);
+}
+
+/**
+ * Test that we can create multiple attributes of different types on the same metadata domain
+ * and they do not interfere with each other.
+ */
+TEST_CASE_METHOD(PCGTests::FPCGBaseTestWithBasicContext, "PCG::AttributeTyped::MultipleAttributesSameData", "[PCG][AttributeTyped]")
+{
+	using namespace PCGAttributeTypedTests;
+
+	UPCGParamData* ParamData = CreateData<UPCGParamData>();
+	FPCGMetadataDomain* Domain = ParamData->MutableMetadata()->GetDefaultMetadataDomain();
+
+	FPCGMetadataAttribute<float>* FloatAttr = Domain->CreateAttribute<float>(FName(TEXT("FloatAttr")), 1.0f, true, false);
+	FPCGMetadataAttribute<int32>* IntAttr = Domain->CreateAttribute<int32>(FName(TEXT("IntAttr")), 42, true, false);
+	FPCGMetadataAttribute<FString>* StrAttr = Domain->CreateAttribute<FString>(FName(TEXT("StrAttr")), FString(TEXT("Hello")), true, false);
+	FPCGMetadataAttribute<FVector>* VecAttr = Domain->CreateAttribute<FVector>(FName(TEXT("VecAttr")), FVector::ZeroVector, true, false);
+
+	REQUIRE_NOT_EQUAL(FloatAttr, nullptr);
+	REQUIRE_NOT_EQUAL(IntAttr, nullptr);
+	REQUIRE_NOT_EQUAL(StrAttr, nullptr);
+	REQUIRE_NOT_EQUAL(VecAttr, nullptr);
+
+	REQUIRE_EQUAL(Domain->GetAttributeCount(), 4);
+
+	PCGMetadataEntryKey EntryKey = Domain->AddEntry();
+
+	FloatAttr->SetValue(EntryKey, 3.14f);
+	IntAttr->SetValue(EntryKey, 100);
+	StrAttr->SetValue(EntryKey, FString(TEXT("World")));
+	VecAttr->SetValue(EntryKey, FVector(1.0, 2.0, 3.0));
+
+	REQUIRE_EQUAL(FloatAttr->GetValueFromItemKey(EntryKey), 3.14f);
+	REQUIRE_EQUAL(IntAttr->GetValueFromItemKey(EntryKey), 100);
+	REQUIRE_EQUAL(StrAttr->GetValueFromItemKey(EntryKey), FString(TEXT("World")));
+	REQUIRE_EQUAL(VecAttr->GetValueFromItemKey(EntryKey), FVector(1.0, 2.0, 3.0));
+
+	// Defaults should be unaffected
+	REQUIRE_EQUAL(FloatAttr->GetValueFromItemKey(PCGInvalidEntryKey), 1.0f);
+	REQUIRE_EQUAL(IntAttr->GetValueFromItemKey(PCGInvalidEntryKey), 42);
+	REQUIRE_EQUAL(StrAttr->GetValueFromItemKey(PCGInvalidEntryKey), FString(TEXT("Hello")));
+	REQUIRE_EQUAL(VecAttr->GetValueFromItemKey(PCGInvalidEntryKey), FVector::ZeroVector);
+}
+
+/**
+ * Test that creating a duplicate attribute with CreateAttribute (not FindOrCreate) produces a warning
+ * but returns the existing attribute when types match.
+ */
+TEST_CASE_METHOD(PCGTests::FPCGBaseTestWithBasicContext, "PCG::AttributeTyped::DuplicateCreateAttribute", "[PCG][AttributeTyped]")
+{
+	using namespace PCGAttributeTypedTests;
+
+	UPCGParamData* ParamData = CreateData<UPCGParamData>();
+	FPCGMetadataDomain* Domain = ParamData->MutableMetadata()->GetDefaultMetadataDomain();
+
+	FPCGMetadataAttribute<float>* Attr1 = Domain->CreateAttribute<float>(AttributeName, 1.0f, true, false);
+	REQUIRE_NOT_EQUAL(Attr1, nullptr);
+
+	// Creating again with the same type should log a warning but succeed
+	{
+		PCGTests::FPCGTestsLogOutputDevice LogCapture(/*bSuppressErrors=*/true);
+		FPCGMetadataAttribute<float>* Attr2 = Domain->CreateAttribute<float>(AttributeName, 2.0f, true, false);
+		// Should return the existing attribute
+		REQUIRE_NOT_EQUAL(Attr2, nullptr);
+	}
+}
+
+/**
+ * Test that a large number of entries works correctly.
+ * This helps catch off-by-one errors and capacity issues.
+ */
+TEST_CASE_METHOD(PCGTests::FPCGBaseTestWithBasicContext, "PCG::AttributeTyped::LargeEntryCount", "[PCG][AttributeTyped]")
+{
+	using namespace PCGAttributeTypedTests;
+
+	constexpr int32 LargeCount = 10000;
+
+	UPCGPointArrayData* PointArrayData = CreateData<UPCGPointArrayData>();
+	PointArrayData->SetNumPoints(LargeCount);
+	PointArrayData->AllocateProperties(EPCGPointNativeProperties::MetadataEntry);
+	TPCGValueRange<int64> MetadataEntryRange = PointArrayData->GetMetadataEntryValueRange(/*bAllocate=*/false);
+
+	FPCGMetadataDomain* Domain = PointArrayData->MutableMetadata()->GetDefaultMetadataDomain();
+
+	FPCGMetadataAttribute<int32>* Attr = Domain->CreateAttribute<int32>(PCGAttributeTypedTests::AttributeName, -1, true, false);
+	REQUIRE_NOT_EQUAL(Attr, nullptr);
+
+	TArray<int32> Values;
+	Values.SetNum(LargeCount);
+
+	for (int32 i = 0; i < LargeCount; ++i)
+	{
+		MetadataEntryRange[i] = Domain->AddEntry(MetadataEntryRange[i]);
+		Values[i] = i;
+	}
+
+	Attr->SetValues(PCGValueRangeHelpers::MakeConstValueRange(MetadataEntryRange), Values);
+	REQUIRE_EQUAL(Attr->GetNumberOfEntries(), LargeCount);
+
+	// Verify first, last, and middle
+	REQUIRE_EQUAL(Attr->GetValueFromItemKey(MetadataEntryRange[0]), 0);
+	REQUIRE_EQUAL(Attr->GetValueFromItemKey(MetadataEntryRange[LargeCount - 1]), LargeCount - 1);
+	REQUIRE_EQUAL(Attr->GetValueFromItemKey(MetadataEntryRange[LargeCount / 2]), LargeCount / 2);
+
+	// Verify all values
+	TArray<int32> ReadValues;
+	ReadValues.SetNum(LargeCount);
+	Attr->GetValuesFromItemKeys(PCGValueRangeHelpers::MakeConstValueRange(MetadataEntryRange), ReadValues);
+
+	for (int32 i = 0; i < LargeCount; ++i)
+	{
+		REQUIRE_EQUAL(ReadValues[i], i);
+	}
+}
+
+/**
+ * Test compressed data behavior for FString attributes: duplicate values should share value keys.
+ */
+TEST_CASE_METHOD(PCGTests::FPCGBaseTestWithBasicContext, "PCG::AttributeTyped::CompressedData", "[PCG][AttributeTyped]")
+{
+	using namespace PCGAttributeTypedTests;
+
+	constexpr int32 Count = 100;
+
+	UPCGPointArrayData* PointArrayData = CreateData<UPCGPointArrayData>();
+	PointArrayData->SetNumPoints(Count);
+	PointArrayData->AllocateProperties(EPCGPointNativeProperties::MetadataEntry);
+	TPCGValueRange<int64> MetadataEntryRange = PointArrayData->GetMetadataEntryValueRange(/*bAllocate=*/false);
+
+	FPCGMetadataDomain* Domain = PointArrayData->MutableMetadata()->GetDefaultMetadataDomain();
+
+	// Use the default value as one of the values we'll set
+	FPCGMetadataAttribute<FString>* Attr = Domain->CreateAttribute<FString>(AttributeName, FString(TEXT("Value: 0")), true, false);
+	REQUIRE_NOT_EQUAL(Attr, nullptr);
+
+	TStaticArray<FString, Count> Values{};
+
+	// First half: 10 unique values (mod 10)
+	for (int32 i = 0; i < Count / 2; ++i)
+	{
+		Domain->InitializeOnSet(MetadataEntryRange[i]);
+		Values[i] = FString::Printf(TEXT("Value: %d"), i % 10);
+	}
+
+	TArray<PCGMetadataEntryKey> FirstHalfKeys;
+	FirstHalfKeys.SetNum(Count / 2);
+	for (int32 i = 0; i < Count / 2; ++i)
+	{
+		FirstHalfKeys[i] = MetadataEntryRange[i];
+	}
+
+	Attr->SetValues(TConstArrayView<PCGMetadataEntryKey>(FirstHalfKeys), MakeArrayView(Values.GetData(), Count / 2));
+
+	// Should have only 9 unique value keys (the tenth one is the default "Value: 0")
+	REQUIRE_EQUAL(Attr->GetValueKeyOffsetForChild(), 9);
+
+	// Second half: mix of existing and new strings
+	TArray<PCGMetadataEntryKey> SecondHalfKeys;
+	SecondHalfKeys.SetNum(Count / 2);
+	for (int32 i = Count / 2; i < Count; ++i)
+	{
+		const int32 Index = i - Count / 2;
+		Domain->InitializeOnSet(MetadataEntryRange[i]);
+		SecondHalfKeys[Index] = MetadataEntryRange[i];
+
+		if (Index < 10)
+		{
+			Values[i] = FString::Printf(TEXT("Value: %d"), Index % 10);
+		}
+		else
+		{
+			Values[i] = FString::Printf(TEXT("New Value: %d"), Index % 10);
+		}
+	}
+
+	Attr->SetValues(TConstArrayView<PCGMetadataEntryKey>(SecondHalfKeys), MakeArrayView(Values.GetData() + Count / 2, Count / 2));
+
+	// Should have 19 unique value keys now (9 old + 10 new, the twentieth one is the default which matches "Value: 0")
+	REQUIRE_EQUAL(Attr->GetValueKeyOffsetForChild(), 19);
+
+	// Verify all values can be read back
+	TStaticArray<FString, Count> ReadValues{};
+	Attr->GetValuesFromItemKeys(PCGValueRangeHelpers::MakeConstValueRange(MetadataEntryRange), ReadValues);
+
+	for (int32 i = 0; i < Count; ++i)
+	{
+		REQUIRE_EQUAL(ReadValues[i], Values[i]);
+	}
+}
+
+/**
+ * Test that deleting an attribute actually removes it from the domain.
+ */
+TEST_CASE_METHOD(PCGTests::FPCGBaseTestWithBasicContext, "PCG::AttributeTyped::DeleteAttribute", "[PCG][AttributeTyped]")
+{
+	using namespace PCGAttributeTypedTests;
+
+	UPCGParamData* ParamData = CreateData<UPCGParamData>();
+	FPCGMetadataDomain* Domain = ParamData->MutableMetadata()->GetDefaultMetadataDomain();
+
+	FPCGMetadataAttribute<float>* Attr = Domain->CreateAttribute<float>(AttributeName, 1.0f, true, false);
+	REQUIRE_NOT_EQUAL(Attr, nullptr);
+	REQUIRE(Domain->HasAttribute(AttributeName));
+	REQUIRE_EQUAL(Domain->GetAttributeCount(), 1);
+
+	Domain->DeleteAttribute(AttributeName);
+	REQUIRE(!Domain->HasAttribute(AttributeName));
+	REQUIRE_EQUAL(Domain->GetAttributeCount(), 0);
+
+	// Should be able to create a new attribute with the same name after deletion
+	FPCGMetadataAttribute<int32>* NewAttr = Domain->CreateAttribute<int32>(AttributeName, 42, true, false);
+	REQUIRE_NOT_EQUAL(NewAttr, nullptr);
+	REQUIRE(IsOfType<int32>(NewAttr));
+}
+
+/**
+ * Test setting values one at a time (single SetValue) vs bulk SetValues, and verify results match.
+ */
+TEST_CASE_METHOD(PCGTests::FPCGBaseTestWithBasicContext, "PCG::AttributeTyped::SingleVsBulkSetValues", "[PCG][AttributeTyped]")
+{
+	using namespace PCGAttributeTypedTests;
+
+	FRandomStream RandomStream(Seed);
+
+	UPCGParamData* ParamData1 = CreateData<UPCGParamData>();
+	UPCGParamData* ParamData2 = CreateData<UPCGParamData>();
+	FPCGMetadataDomain* Domain1 = ParamData1->MutableMetadata()->GetDefaultMetadataDomain();
+	FPCGMetadataDomain* Domain2 = ParamData2->MutableMetadata()->GetDefaultMetadataDomain();
+
+	FPCGMetadataAttribute<FVector>* SingleAttr = Domain1->CreateAttribute<FVector>(AttributeName, FVector::ZeroVector, true, false);
+	FPCGMetadataAttribute<FVector>* BulkAttr = Domain2->CreateAttribute<FVector>(AttributeName, FVector::ZeroVector, true, false);
+
+	REQUIRE_NOT_EQUAL(SingleAttr, nullptr);
+	REQUIRE_NOT_EQUAL(BulkAttr, nullptr);
+
+	constexpr int32 Count = 20;
+	TStaticArray<FVector, Count> Values{};
+	TStaticArray<PCGMetadataEntryKey, Count> EntryKeys1;
+	TStaticArray<PCGMetadataEntryKey, Count> EntryKeys2;
+
+	for (int32 i = 0; i < Count; ++i)
+	{
+		EntryKeys1[i] = Domain1->AddEntry();
+		EntryKeys2[i] = Domain2->AddEntry();
+		Values[i] = GenerateRandom<FVector>(RandomStream);
+	}
+
+	// Single set
+	for (int32 i = 0; i < Count; ++i)
+	{
+		SingleAttr->SetValue(EntryKeys1[i], Values[i]);
+	}
+
+	// Bulk set
+	BulkAttr->SetValues(TConstArrayView<PCGMetadataEntryKey>(EntryKeys2), Values);
+
+	// Read both and compare
+	TStaticArray<FVector, Count> ReadSingle{};
+	TStaticArray<FVector, Count> ReadBulk{};
+	SingleAttr->GetValuesFromItemKeys(TConstArrayView<PCGMetadataEntryKey>(EntryKeys1), ReadSingle);
+	BulkAttr->GetValuesFromItemKeys(TConstArrayView<PCGMetadataEntryKey>(EntryKeys2), ReadBulk);
+
+	for (int32 i = 0; i < Count; ++i)
+	{
+		REQUIRE_EQUAL(ReadSingle[i], Values[i]);
+		REQUIRE_EQUAL(ReadBulk[i], Values[i]);
+		REQUIRE_EQUAL(ReadSingle[i], ReadBulk[i]);
+	}
+}
+
+/**
+ * Test that unset entry keys return the default value.
+ */
+TEST_CASE_METHOD(PCGTests::FPCGBaseTestWithBasicContext, "PCG::AttributeTyped::UnsetEntryReturnsDefault", "[PCG][AttributeTyped]")
+{
+	using namespace PCGAttributeTypedTests;
+
+	UPCGParamData* ParamData = CreateData<UPCGParamData>();
+	FPCGMetadataDomain* Domain = ParamData->MutableMetadata()->GetDefaultMetadataDomain();
+
+	FPCGMetadataAttribute<double>* Attr = Domain->CreateAttribute<double>(AttributeName, 99.9, true, false);
+	REQUIRE_NOT_EQUAL(Attr, nullptr);
+
+	// Add entries but don't set any values on them
+	PCGMetadataEntryKey Key1 = Domain->AddEntry();
+	PCGMetadataEntryKey Key2 = Domain->AddEntry();
+
+	// Reading unset entries should return the default value
+	REQUIRE_EQUAL(Attr->GetValueFromItemKey(Key1), 99.9);
+	REQUIRE_EQUAL(Attr->GetValueFromItemKey(Key2), 99.9);
+
+	// Set one, the other should still be default
+	Attr->SetValue(Key1, 1.0);
+	REQUIRE_EQUAL(Attr->GetValueFromItemKey(Key1), 1.0);
+	REQUIRE_EQUAL(Attr->GetValueFromItemKey(Key2), 99.9);
+}
+
+/**
+ * Test that AddValues (bulk) returns correct value keys and that they can be used with GetValue.
+ */
+TEST_CASE_METHOD(PCGTests::FPCGBaseTestWithBasicContext, "PCG::AttributeTyped::BulkAddValues", "[PCG][AttributeTyped]")
+{
+	using namespace PCGAttributeTypedTests;
+
+	UPCGParamData* ParamData = CreateData<UPCGParamData>();
+	FPCGMetadataDomain* Domain = ParamData->MutableMetadata()->GetDefaultMetadataDomain();
+
+	FPCGMetadataAttribute<int32>* Attr = Domain->CreateAttribute<int32>(AttributeName, 0, true, false);
+	REQUIRE_NOT_EQUAL(Attr, nullptr);
+
+	TArray<int32> ValuesBatch = {10, 20, 30, 40, 50};
+	TArray<PCGMetadataValueKey> ValueKeys = Attr->AddValues(ValuesBatch);
+	REQUIRE_EQUAL(ValueKeys.Num(), ValuesBatch.Num());
+
+	for (int32 i = 0; i < ValuesBatch.Num(); ++i)
+	{
+		int32 ReadVal = Attr->GetValue(ValueKeys[i]);
+		REQUIRE_EQUAL(ReadVal, ValuesBatch[i]);
+	}
+}
+
+/**
+ * Test that GetValues with multiple value keys works correctly in bulk.
+ */
+TEST_CASE_METHOD(PCGTests::FPCGBaseTestWithBasicContext, "PCG::AttributeTyped::BulkGetValues", "[PCG][AttributeTyped]")
+{
+	using namespace PCGAttributeTypedTests;
+
+	UPCGParamData* ParamData = CreateData<UPCGParamData>();
+	FPCGMetadataDomain* Domain = ParamData->MutableMetadata()->GetDefaultMetadataDomain();
+
+	FPCGMetadataAttribute<float>* Attr = Domain->CreateAttribute<float>(AttributeName, 0.0f, true, false);
+	REQUIRE_NOT_EQUAL(Attr, nullptr);
+
+	TArray<float> ValuesBatch = {1.0f, 2.0f, 3.0f, 4.0f};
+	TArray<PCGMetadataValueKey> ValueKeys = Attr->AddValues(ValuesBatch);
+
+	TArray<float> ReadValues;
+	ReadValues.SetNum(ValueKeys.Num());
+	Attr->GetValues(ValueKeys, ReadValues);
+
+	for (int32 i = 0; i < ValuesBatch.Num(); ++i)
+	{
+		REQUIRE_EQUAL(ReadValues[i], ValuesBatch[i]);
+	}
+}
+
+// @todo_pcg to be enabled when typed attribute inherit from generic
+///**
+// * Test that the typed attribute inherits IsGeneric() == true from FPCGMetadataAttributeGeneric.
+// */
+//TEST_CASE_METHOD(PCGTests::FPCGBaseTestWithBasicContext, "PCG::AttributeTyped::IsGenericFlag", "[PCG][AttributeTyped]")
+//{
+//	using namespace PCGAttributeTypedTests;
+//
+//	UPCGParamData* ParamData = CreateData<UPCGParamData>();
+//	FPCGMetadataDomain* Domain = ParamData->MutableMetadata()->GetDefaultMetadataDomain();
+//
+//	FPCGMetadataAttribute<float>* Attr = Domain->CreateAttribute<float>(AttributeName, 0.0f, true, false);
+//	REQUIRE_NOT_EQUAL(Attr, nullptr);
+//
+//	// FPCGMetadataAttribute<T> inherits from FPCGMetadataAttributeGeneric, so IsGeneric should be true
+//	REQUIRE(Attr->IsGeneric());
+//}
+
+/**
+ * Test CopyExistingAttribute through the domain for typed attributes.
+ */
+TEST_CASE_METHOD(PCGTests::FPCGBaseTestWithBasicContext, "PCG::AttributeTyped::CopyExistingAttribute", "[PCG][AttributeTyped]")
+{
+	using namespace PCGAttributeTypedTests;
+
+	UPCGParamData* ParamData = CreateData<UPCGParamData>();
+	FPCGMetadataDomain* Domain = ParamData->MutableMetadata()->GetDefaultMetadataDomain();
+
+	FPCGMetadataAttribute<FVector>* Attr = Domain->CreateAttribute<FVector>(AttributeName, FVector(1.0, 2.0, 3.0), true, false);
+	REQUIRE_NOT_EQUAL(Attr, nullptr);
+
+	PCGMetadataEntryKey EntryKey = Domain->AddEntry();
+	Attr->SetValue(EntryKey, FVector(10.0, 20.0, 30.0));
+
+	FName CopyName(TEXT("CopiedVecAttr"));
+	bool bCopied = Domain->CopyExistingAttribute(AttributeName, CopyName, /*bKeepParent=*/false);
+	REQUIRE(bCopied);
+
+	FPCGMetadataAttribute<FVector>* CopiedAttr = Domain->GetMutableTypedAttribute<FVector>(CopyName);
+	REQUIRE_NOT_EQUAL(CopiedAttr, nullptr);
+
+	// Default should match
+	REQUIRE_EQUAL(CopiedAttr->GetValueFromItemKey(PCGInvalidEntryKey), FVector(1.0, 2.0, 3.0));
+
+	// Values should match
+	REQUIRE_EQUAL(CopiedAttr->GetValueFromItemKey(EntryKey), FVector(10.0, 20.0, 30.0));
+}
+
+/**
+ * Test RenameAttribute through the domain for typed attributes.
+ */
+TEST_CASE_METHOD(PCGTests::FPCGBaseTestWithBasicContext, "PCG::AttributeTyped::RenameAttribute", "[PCG][AttributeTyped]")
+{
+	using namespace PCGAttributeTypedTests;
+
+	UPCGParamData* ParamData = CreateData<UPCGParamData>();
+	FPCGMetadataDomain* Domain = ParamData->MutableMetadata()->GetDefaultMetadataDomain();
+
+	FPCGMetadataAttribute<int32>* Attr = Domain->CreateAttribute<int32>(AttributeName, 7, true, false);
+	REQUIRE_NOT_EQUAL(Attr, nullptr);
+
+	PCGMetadataEntryKey EntryKey = Domain->AddEntry();
+	Attr->SetValue(EntryKey, 42);
+
+	FName NewName(TEXT("RenamedAttr"));
+	bool bRenamed = Domain->RenameAttribute(AttributeName, NewName);
+	REQUIRE(bRenamed);
+
+	// Old name should not exist
+	REQUIRE(!Domain->HasAttribute(AttributeName));
+
+	// New name should exist and have the correct values
+	FPCGMetadataAttribute<int32>* RenamedAttr = Domain->GetMutableTypedAttribute<int32>(NewName);
+	REQUIRE_NOT_EQUAL(RenamedAttr, nullptr);
+	REQUIRE_EQUAL(RenamedAttr->GetValueFromItemKey(PCGInvalidEntryKey), 7);
+	REQUIRE_EQUAL(RenamedAttr->GetValueFromItemKey(EntryKey), 42);
+}
diff --git a/Engine/Plugins/PCG/Tests/Private/Metadata/PCGMetadataTest.cpp b/Engine/Plugins/PCG/Tests/Private/Metadata/PCGMetadataTest.cpp
new file mode 100644
--- /dev/null
+++ b/Engine/Plugins/PCG/Tests/Private/Metadata/PCGMetadataTest.cpp
@@ -0,0 +1,325 @@
+// Copyright Epic Games, Inc. All Rights Reserved.
+
+#include "PCGTestsCommon.h"
+
+#include "PCGParamData.h"
+#include "Data/PCGPointArrayData.h"
+#include "Metadata/PCGMetadata.h"
+#include "Metadata/PCGMetadataAttributeTpl.h"
+
+#include <catch2/catch_test_macros.hpp>
+
+#include "Containers/StaticArray.h"
+#include "TestHarness.h"
+#include "Engine/Engine.h"
+
+/**
+ * Test metadata inheritance chain with pointer-based parenting through point data.
+ * Creates PointData (5 pts, values 0-4) -> PointData2 inherits, adds pts 5-9 -> PointData3 inherits.
+ * Verifies all 10 values are readable from PointData3.
+ *
+ * Note: The original automation test forced a GC run between setup and verification to test that
+ * metadata chains survive parent object collection. This cannot be replicated in low-level Catch2 tests
+ * because objects created through the test context are kept alive for the test duration.
+ */
+TEST_CASE_METHOD(PCGTests::FPCGBaseTestWithBasicContext, "PCG::Metadata::Inherit", "[PCG][Metadata]")
+{
+	const FName AttributeName = TEXT("MyAttr");
+
+	UPCGPointArrayData* PointData = CreateData<UPCGPointArrayData>();
+	PointData->SetNumPoints(5);
+	UPCGMetadata* PointMetadata = PointData->MutableMetadata();
+	FPCGMetadataAttribute<int32>* Attribute = PointMetadata->CreateAttribute<int32>(AttributeName, 0, true, false);
+
+	TPCGValueRange<int64> MetadataEntryRange = PointData->GetMetadataEntryValueRange();
+	for (int32 i = 0; i < PointData->GetNumPoints(); ++i)
+	{
+		PointMetadata->InitializeOnSet(MetadataEntryRange[i]);
+		Attribute->SetValue(MetadataEntryRange[i], i);
+	}
+
+	UPCGPointArrayData* PointData2 = CreateData<UPCGPointArrayData>();
+	UPCGMetadata* PointMetadata2 = PointData2->MutableMetadata();
+	PointMetadata2->Initialize(PointMetadata);
+
+	// Verify that the initialization inherited (pointer-based, no local entries)
+	REQUIRE_EQUAL(PointMetadata2->GetLocalItemCount(), (int64)0);
+
+	FPCGMetadataAttribute<int32>* Attribute2 = PointMetadata2->GetMutableTypedAttribute<int32>(AttributeName);
+	REQUIRE_NOT_EQUAL(Attribute2, nullptr);
+
+	UPCGBasePointData::SetPoints(PointData, PointData2, {}, /*bCopyAll=*/true);
+	PointData2->SetNumPoints(10);
+
+	TPCGValueRange<int64> MetadataEntryRange2 = PointData2->GetMetadataEntryValueRange();
+	for (int32 i = PointData->GetNumPoints(); i < PointData2->GetNumPoints(); ++i)
+	{
+		PointMetadata2->InitializeOnSet(MetadataEntryRange2[i]);
+		Attribute2->SetValue(MetadataEntryRange2[i], i);
+	}
+
+	UPCGPointArrayData* PointData3 = CreateData<UPCGPointArrayData>();
+	UPCGMetadata* PointMetadata3 = PointData3->MutableMetadata();
+	PointMetadata3->Initialize(PointMetadata2);
+
+	REQUIRE_EQUAL(PointMetadata3->GetLocalItemCount(), (int64)0);
+
+	UPCGBasePointData::SetPoints(PointData2, PointData3, {}, /*bCopyAll=*/true);
+
+	REQUIRE_EQUAL(PointData2->GetNumPoints(), PointData3->GetNumPoints());
+
+	// Verify that all the values are readable from the metadata
+	const FPCGMetadataAttribute<int32>* Attribute3 = PointData3->ConstMetadata()->GetConstTypedAttribute<int32>(AttributeName);
+	REQUIRE_NOT_EQUAL(Attribute3, nullptr);
+
+	TConstPCGValueRange<int64> MetadataEntryRange3 = PointData3->GetConstMetadataEntryValueRange();
+	for (int32 i = 0; i < PointData3->GetNumPoints(); ++i)
+	{
+		REQUIRE_EQUAL(Attribute3->GetValueFromItemKey(MetadataEntryRange3[i]), i);
+	}
+}
+
+/**
+ * Test metadata inheritance with copy-based parenting (InitializeAsCopy).
+ * Same chain as Inherit, but PointData3 uses InitializeAsCopy instead of Initialize.
+ * This means PointData3 gets its own copy of entries rather than pointing to a parent.
+ */
+TEST_CASE_METHOD(PCGTests::FPCGBaseTestWithBasicContext, "PCG::Metadata::InheritCopy", "[PCG][Metadata]")
+{
+	const FName AttributeName = TEXT("MyAttr");
+
+	UPCGPointArrayData* PointData = CreateData<UPCGPointArrayData>();
+	PointData->SetNumPoints(5);
+	UPCGMetadata* PointMetadata = PointData->MutableMetadata();
+	FPCGMetadataAttribute<int32>* Attribute = PointMetadata->CreateAttribute<int32>(AttributeName, 0, true, false);
+
+	TPCGValueRange<int64> MetadataEntryRange = PointData->GetMetadataEntryValueRange();
+	for (int32 i = 0; i < PointData->GetNumPoints(); ++i)
+	{
+		PointMetadata->InitializeOnSet(MetadataEntryRange[i]);
+		Attribute->SetValue(MetadataEntryRange[i], i);
+	}
+
+	UPCGPointArrayData* PointData2 = CreateData<UPCGPointArrayData>();
+	UPCGMetadata* PointMetadata2 = PointData2->MutableMetadata();
+	PointMetadata2->Initialize(PointMetadata);
+
+	// Verify that the initialization inherited (pointer-based, no local entries)
+	REQUIRE_EQUAL(PointMetadata2->GetLocalItemCount(), (int64)0);
+
+	FPCGMetadataAttribute<int32>* Attribute2 = PointMetadata2->GetMutableTypedAttribute<int32>(AttributeName);
+	REQUIRE_NOT_EQUAL(Attribute2, nullptr);
+
+	UPCGBasePointData::SetPoints(PointData, PointData2, {}, /*bCopyAll=*/true);
+	PointData2->SetNumPoints(10);
+
+	TPCGValueRange<int64> MetadataEntryRange2 = PointData2->GetMetadataEntryValueRange();
+	for (int32 i = PointData->GetNumPoints(); i < PointData2->GetNumPoints(); ++i)
+	{
+		PointMetadata2->InitializeOnSet(MetadataEntryRange2[i]);
+		Attribute2->SetValue(MetadataEntryRange2[i], i);
+	}
+
+	UPCGPointArrayData* PointData3 = CreateData<UPCGPointArrayData>();
+	UPCGMetadata* PointMetadata3 = PointData3->MutableMetadata();
+	PointMetadata3->InitializeAsCopy(FPCGMetadataInitializeParams(PointMetadata2));
+
+	// Verify that the initialization made a copy (PointData2 had 5 local entries for indices 5-9)
+	REQUIRE_EQUAL(PointMetadata3->GetLocalItemCount(), (int64)5);
+
+	UPCGBasePointData::SetPoints(PointData2, PointData3, {}, /*bCopyAll=*/true);
+
+	REQUIRE_EQUAL(PointData2->GetNumPoints(), PointData3->GetNumPoints());
+
+	// Verify that all the values are readable from the metadata
+	const FPCGMetadataAttribute<int32>* Attribute3 = PointData3->ConstMetadata()->GetConstTypedAttribute<int32>(AttributeName);
+	REQUIRE_NOT_EQUAL(Attribute3, nullptr);
+
+	TConstPCGValueRange<int64> MetadataEntryRange3 = PointData3->GetConstMetadataEntryValueRange();
+	for (int32 i = 0; i < PointData3->GetNumPoints(); ++i)
+	{
+		REQUIRE_EQUAL(Attribute3->GetValueFromItemKey(MetadataEntryRange3[i]), i);
+	}
+}
+
+/**
+ * Test metadata inheritance with ParamData, which does not support pointer-based parenting.
+ * Initialize always copies entries for ParamData.
+ * Creates ParamData (5 entries, values 0-4) -> ParamData2 (copies + 5 more, values 0-9) -> ParamData3 (copies all 10).
+ */
+TEST_CASE_METHOD(PCGTests::FPCGBaseTestWithBasicContext, "PCG::Metadata::InheritWithNoParenting", "[PCG][Metadata]")
+{
+	const FName AttributeName = TEXT("MyAttr");
+
+	UPCGParamData* ParamData = CreateData<UPCGParamData>();
+	UPCGMetadata* ParamMetadata = ParamData->MutableMetadata();
+	FPCGMetadataAttribute<int32>* Attribute = ParamMetadata->CreateAttribute<int32>(AttributeName, 0, true, false);
+
+	for (int32 i = 0; i < 5; ++i)
+	{
+		PCGMetadataEntryKey EntryKey = PCGInvalidEntryKey;
+		ParamMetadata->InitializeOnSet(EntryKey);
+		Attribute->SetValue(EntryKey, i);
+	}
+
+	UPCGParamData* ParamData2 = CreateData<UPCGParamData>();
+	UPCGMetadata* ParamMetadata2 = ParamData2->MutableMetadata();
+	ParamMetadata2->Initialize(ParamMetadata);
+
+	// Verify that the initialization made a copy (ParamData doesn't support pointer-based parenting)
+	REQUIRE_EQUAL(ParamMetadata2->GetLocalItemCount(), (int64)5);
+
+	FPCGMetadataAttribute<int32>* Attribute2 = ParamMetadata2->GetMutableTypedAttribute<int32>(AttributeName);
+	REQUIRE_NOT_EQUAL(Attribute2, nullptr);
+
+	for (int32 i = 5; i < 10; ++i)
+	{
+		PCGMetadataEntryKey EntryKey = PCGInvalidEntryKey;
+		ParamMetadata2->InitializeOnSet(EntryKey);
+		Attribute2->SetValue(EntryKey, i);
+	}
+
+	UPCGParamData* ParamData3 = CreateData<UPCGParamData>();
+	UPCGMetadata* ParamMetadata3 = ParamData3->MutableMetadata();
+	ParamMetadata3->Initialize(ParamMetadata2);
+
+	// Verify that the initialization made a copy
+	REQUIRE_EQUAL(ParamMetadata3->GetLocalItemCount(), (int64)10);
+
+	// Verify that all the values are readable from the metadata
+	const FPCGMetadataAttribute<int32>* Attribute3 = ParamData3->ConstMetadata()->GetConstTypedAttribute<int32>(AttributeName);
+	REQUIRE_NOT_EQUAL(Attribute3, nullptr);
+
+	for (int32 i = 0; i < 10; ++i)
+	{
+		REQUIRE_EQUAL(Attribute3->GetValueFromItemKey(i), i);
+	}
+}
+
+/**
+ * Test that the Data metadata domain only supports a single entry.
+ * Adding two entries via InitializeOnSet should produce an error and keep only one entry with the latest value.
+ */
+TEST_CASE_METHOD(PCGTests::FPCGBaseTestWithBasicContext, "PCG::Metadata::TwoAddEntryToData", "[PCG][Metadata]")
+{
+	const FName AttributeName = TEXT("MyAttr");
+
+	UPCGParamData* ParamData = CreateData<UPCGParamData>();
+	UPCGMetadata* ParamMetadata = ParamData->MutableMetadata();
+	FPCGMetadataDomain* ParamMetadataDomain = ParamMetadata->GetMetadataDomain(PCGMetadataDomainID::Data);
+	FPCGMetadataAttribute<int32>* Attribute = ParamMetadataDomain->CreateAttribute<int32>(AttributeName, 0, false, false);
+	REQUIRE_NOT_EQUAL(Attribute, nullptr);
+
+	{
+		PCGTests::FPCGTestsLogOutputDevice LogCapture(/*bSuppressErrors=*/true);
+
+		for (int32 i = 0; i < 2; ++i)
+		{
+			PCGMetadataEntryKey EntryKey = PCGInvalidEntryKey;
+			ParamMetadataDomain->InitializeOnSet(EntryKey);
+			Attribute->SetValue(EntryKey, i);
+		}
+
+		// Second InitializeOnSet should have produced a warning since Data domain only supports single entries
+		REQUIRE(LogCapture.NbMessageReceived.Contains(ELogVerbosity::Warning));
+	}
+
+	REQUIRE_EQUAL(ParamMetadataDomain->GetItemCountForChild(), (int64)1);
+	REQUIRE_EQUAL(Attribute->GetValueFromItemKey(PCGFirstEntryKey), 1);
+}
+
+/**
+ * Test that the Data metadata domain rejects bulk AddEntries and only keeps a single entry.
+ */
+TEST_CASE_METHOD(PCGTests::FPCGBaseTestWithBasicContext, "PCG::Metadata::AddEntriesToData", "[PCG][Metadata]")
+{
+	const FName AttributeName = TEXT("MyAttr");
+
+	UPCGParamData* ParamData = CreateData<UPCGParamData>();
+	UPCGMetadata* ParamMetadata = ParamData->MutableMetadata();
+	FPCGMetadataDomain* ParamMetadataDomain = ParamMetadata->GetMetadataDomain(PCGMetadataDomainID::Data);
+	FPCGMetadataAttribute<int32>* Attribute = ParamMetadataDomain->CreateAttribute<int32>(AttributeName, 0, false, false);
+	REQUIRE_NOT_EQUAL(Attribute, nullptr);
+
+	{
+		PCGTests::FPCGTestsLogOutputDevice LogCapture(/*bSuppressErrors=*/true);
+
+		TStaticArray<PCGMetadataEntryKey, 2> ParentKeys = {PCGInvalidEntryKey, PCGInvalidEntryKey};
+		TArray<PCGMetadataEntryKey> NewKeys = ParamMetadataDomain->AddEntries(ParentKeys);
+		TStaticArray<int32, 2> Values = {5, 6};
+		Attribute->SetValues(NewKeys, Values);
+
+		REQUIRE(LogCapture.NbMessageReceived.Contains(ELogVerbosity::Warning));
+	}
+
+	REQUIRE_EQUAL(ParamMetadataDomain->GetItemCountForChild(), (int64)1);
+	REQUIRE_EQUAL(Attribute->GetValueFromItemKey(PCGFirstEntryKey), 6);
+}
+
+/**
+ * Test that the Data metadata domain rejects AddEntriesInPlace and only keeps a single entry.
+ */
+TEST_CASE_METHOD(PCGTests::FPCGBaseTestWithBasicContext, "PCG::Metadata::AddEntriesInPlaceToData", "[PCG][Metadata]")
+{
+	const FName AttributeName = TEXT("MyAttr");
+
+	UPCGParamData* ParamData = CreateData<UPCGParamData>();
+	UPCGMetadata* ParamMetadata = ParamData->MutableMetadata();
+	FPCGMetadataDomain* ParamMetadataDomain = ParamMetadata->GetMetadataDomain(PCGMetadataDomainID::Data);
+	FPCGMetadataAttribute<int32>* Attribute = ParamMetadataDomain->CreateAttribute<int32>(AttributeName, 0, false, false);
+	REQUIRE_NOT_EQUAL(Attribute, nullptr);
+
+	{
+		PCGTests::FPCGTestsLogOutputDevice LogCapture(/*bSuppressErrors=*/true);
+
+		TStaticArray<PCGMetadataEntryKey, 2> Keys = {PCGInvalidEntryKey, PCGInvalidEntryKey};
+		TStaticArray<PCGMetadataEntryKey*, 2> KeysPtr = {&Keys[0], &Keys[1]};
+
+		ParamMetadataDomain->AddEntriesInPlace(KeysPtr);
+		TStaticArray<int32, 2> Values = {5, 6};
+		Attribute->SetValues(Keys, Values);
+
+		REQUIRE(LogCapture.NbMessageReceived.Contains(ELogVerbosity::Warning));
+	}
+
+	REQUIRE_EQUAL(ParamMetadataDomain->GetItemCountForChild(), (int64)1);
+	REQUIRE_EQUAL(Attribute->GetValueFromItemKey(PCGFirstEntryKey), 6);
+}
+
+/**
+ * Test that the Data metadata domain rejects delayed entries (AddEntryPlaceholder + AddDelayedEntries)
+ * and only keeps a single entry.
+ */
+TEST_CASE_METHOD(PCGTests::FPCGBaseTestWithBasicContext, "PCG::Metadata::DelayedAddEntriesToData", "[PCG][Metadata]")
+{
+	const FName AttributeName = TEXT("MyAttr");
+
+	UPCGParamData* ParamData = CreateData<UPCGParamData>();
+	UPCGMetadata* ParamMetadata = ParamData->MutableMetadata();
+	FPCGMetadataDomain* ParamMetadataDomain = ParamMetadata->GetMetadataDomain(PCGMetadataDomainID::Data);
+	FPCGMetadataAttribute<int32>* Attribute = ParamMetadataDomain->CreateAttribute<int32>(AttributeName, 0, false, false);
+	REQUIRE_NOT_EQUAL(Attribute, nullptr);
+
+	{
+		PCGTests::FPCGTestsLogOutputDevice LogCapture(/*bSuppressErrors=*/true);
+
+		TArray<TTuple<PCGMetadataEntryKey, PCGMetadataEntryKey>> Mapping;
+		TStaticArray<PCGMetadataEntryKey, 2> Keys;
+
+		Keys[0] = ParamMetadataDomain->AddEntryPlaceholder();
+		Mapping.Emplace(PCGInvalidEntryKey, Keys[0]);
+		Keys[1] = ParamMetadataDomain->AddEntryPlaceholder();
+		Mapping.Emplace(PCGInvalidEntryKey, Keys[1]);
+
+		TStaticArray<int32, 2> Values = {5, 6};
+		Attribute->SetValues(Keys, Values);
+
+		ParamMetadataDomain->AddDelayedEntries(Mapping);
+
+		// Both AddEntryPlaceholder (second call) and AddDelayedEntries should produce warnings
+		REQUIRE(LogCapture.NbMessageReceived.Contains(ELogVerbosity::Warning));
+	}
+
+	REQUIRE_EQUAL(ParamMetadataDomain->GetItemCountForChild(), (int64)1);
+	REQUIRE_EQUAL(Attribute->GetValueFromItemKey(PCGFirstEntryKey), 6);
+}
