[PCG] Fixed PCG Test build error from missing file update

#Synced-CL 51636404

diff --git a/Engine/Plugins/PCG/Tests/Private/Metadata/PCGMetadataPartitionCommonTests.cpp b/Engine/Plugins/PCG/Tests/Private/Metadata/PCGMetadataPartitionCommonTests.cpp
--- a/Engine/Plugins/PCG/Tests/Private/Metadata/PCGMetadataPartitionCommonTests.cpp
+++ b/Engine/Plugins/PCG/Tests/Private/Metadata/PCGMetadataPartitionCommonTests.cpp
@@ -6,7 +6,6 @@
 #include "PCGParamData.h"
 #include "Data/PCGPointArrayData.h"
 #include "Metadata/PCGMetadata.h"
-#include "Metadata/PCGMetadataAttribute.h"
 #include "Elements/Metadata/PCGMetadataPartition.h"
 #include "Metadata/PCGMetadataPartitionCommon.h"
 #include "Metadata/Accessors/PCGAttributeAccessorHelpers.h"
@@ -14,6 +13,7 @@
 #include <catch2/catch_test_macros.hpp>
 
 #include "TestHarness.h"
+#include "Misc/ScopedCVar.h"
 
 namespace PCGMetadataPartitionCommonTests
 {
@@ -330,7 +330,7 @@ TEST_CASE_METHOD(PCGTests::FPCGBaseTestWithBasicContext, "PCG::AttributePartitio
 		Selectors.Add(FPCGAttributePropertySelector::CreateAttributeSelector(AttrA));
 		Selectors.Add(FPCGAttributePropertySelector::CreateAttributeSelector(AttrB));
 
-		PCGTests::FCVarScope DisableMixedRadix(TEXT("pcg.MetadataPartition.UseMixedRadixHash"), 0);
+		FScopedCVar<int> DisableMixedRadix(TEXT("pcg.MetadataPartition.UseMixedRadixHash"), 0);
 
 		TArray<TArray<int32>> Result = PCGMetadataPartitionCommon::AttributeGenericPartition(Data, Selectors, &Context);
 		VerifyGenericPartition(Result, {{0}, {1}, {2}, {3, 4}});
