[PCG] Refactor execution cache so it can be extended and requires less boiler plate code
- Add unittest

#rnx
#Synced-CL 51757879

#Synced-CL 51757879

diff --git a/Engine/Plugins/PCG/Source/PCG/Private/Graph/PCGGraphExecutor.cpp b/Engine/Plugins/PCG/Source/PCG/Private/Graph/PCGGraphExecutor.cpp
--- a/Engine/Plugins/PCG/Source/PCG/Private/Graph/PCGGraphExecutor.cpp
+++ b/Engine/Plugins/PCG/Source/PCG/Private/Graph/PCGGraphExecutor.cpp
@@ -3261,8 +3261,11 @@ bool FPCGPreGraphElement::ExecuteInternal(FPCGContext* Context) const
 		return true;
 	}
 
-	ExecutionSource->GetExecutionState().ExecutePreGraph(Context);
-
+	{
+		PCG_EXECUTION_CACHE_VALIDATION_CREATE_SCOPE(ExecutionSource);
+		ExecutionSource->GetExecutionState().ExecutePreGraph(Context);
+	}
+	
 	// Also emplace in the output a copy of the user parameters. Check in priority: Graph override, graph instance on the execution source and graph on the execution source.
 	const UPCGGraphInterface* Graph = GraphOverride.Get();
 
diff --git a/Engine/Plugins/PCG/Source/PCG/Private/Graph/PCGGraphExecutor.h b/Engine/Plugins/PCG/Source/PCG/Private/Graph/PCGGraphExecutor.h
--- a/Engine/Plugins/PCG/Source/PCG/Private/Graph/PCGGraphExecutor.h
+++ b/Engine/Plugins/PCG/Source/PCG/Private/Graph/PCGGraphExecutor.h
@@ -4,12 +4,12 @@
 
 #include "PCGContext.h"
 #include "PCGElement.h"
-#include "PCGGraphPerExecutionCache.h"
 #include "PCGPin.h"
 #include "PCGSettings.h"
 #include "PCGSubsystem.h"
 #include "Graph/PCGGraphCache.h"
 #include "Graph/PCGGraphCompiler.h"
+#include "Graph/PCGGraphPerExecutionCache.h"
 #include "Graph/PCGGraphTask.h"
 #include "Graph/PCGStackContext.h"
 
@@ -164,30 +164,6 @@ private:
 
 	FPCGPerExecutionCache PerExecutionCache;
 
-	UPCGData* GetPCGData(FPCGTaskId InGraphExecutionTaskId, bool& bOutFound) { return PerExecutionCache.GetExecutionCachePCGData(InGraphExecutionTaskId, EPCGPerExecutionCacheDataType::PCGData, bOutFound); }
-	UPCGData* GetInputPCGData(FPCGTaskId InGraphExecutionTaskId, bool& bOutFound) { return PerExecutionCache.GetExecutionCachePCGData(InGraphExecutionTaskId, EPCGPerExecutionCacheDataType::InputData, bOutFound); }
-	UPCGData* GetSelfPCGData(FPCGTaskId InGraphExecutionTaskId, bool& bOutFound) { return PerExecutionCache.GetExecutionCachePCGData(InGraphExecutionTaskId, EPCGPerExecutionCacheDataType::SelfData, bOutFound); }
-	UPCGData* GetLandscapePCGData(FPCGTaskId InGraphExecutionTaskId, bool& bOutFound) { return PerExecutionCache.GetExecutionCachePCGData(InGraphExecutionTaskId, EPCGPerExecutionCacheDataType::LandscapeData, bOutFound); }
-	UPCGData* GetLandscapeHeightPCGData(FPCGTaskId InGraphExecutionTaskId, bool& bOutFound) { return PerExecutionCache.GetExecutionCachePCGData(InGraphExecutionTaskId, EPCGPerExecutionCacheDataType::LandscapeHeightData, bOutFound); }
-	UPCGData* GetOriginalSelfPCGData(FPCGTaskId InGraphExecutionTaskId, bool& bOutFound) { return PerExecutionCache.GetExecutionCachePCGData(InGraphExecutionTaskId, EPCGPerExecutionCacheDataType::OriginalSelfData, bOutFound); }
-	TOptional<FBox> GetBounds(FPCGTaskId InGraphExecutionTaskId) { return PerExecutionCache.GetExecutionCacheBounds(InGraphExecutionTaskId, EPCGPerExecutionCacheDataType::Bounds); }
-	TOptional<FBox> GetOriginalBounds(FPCGTaskId InGraphExecutionTaskId) { return PerExecutionCache.GetExecutionCacheBounds(InGraphExecutionTaskId, EPCGPerExecutionCacheDataType::OriginalBounds); }
-	TOptional<FBox> GetLocalSpaceBounds(FPCGTaskId InGraphExecutionTaskId) { return PerExecutionCache.GetExecutionCacheBounds(InGraphExecutionTaskId, EPCGPerExecutionCacheDataType::LocalSpaceBounds); }
-	TOptional<FBox> GetOriginalLocalSpaceBounds(FPCGTaskId InGraphExecutionTaskId) { return PerExecutionCache.GetExecutionCacheBounds(InGraphExecutionTaskId, EPCGPerExecutionCacheDataType::OriginalLocalSpaceBounds); }
-	TOptional<FBox> GetTotalBounds(FPCGTaskId InGraphExecutionTaskId) { return PerExecutionCache.GetExecutionCacheBounds(InGraphExecutionTaskId, EPCGPerExecutionCacheDataType::TotalBounds); }
-		
-	void SetPCGData(FPCGTaskId InGraphExecutionTaskId, UPCGData* InData) { PerExecutionCache.SetExecutionCachePCGData(InGraphExecutionTaskId, EPCGPerExecutionCacheDataType::PCGData, InData); }
-	void SetInputPCGData(FPCGTaskId InGraphExecutionTaskId, UPCGData* InData) { PerExecutionCache.SetExecutionCachePCGData(InGraphExecutionTaskId, EPCGPerExecutionCacheDataType::InputData, InData); }
-	void SetSelfPCGData(FPCGTaskId InGraphExecutionTaskId, UPCGData* InData) { PerExecutionCache.SetExecutionCachePCGData(InGraphExecutionTaskId, EPCGPerExecutionCacheDataType::SelfData, InData); }
-	void SetLandscapePCGData(FPCGTaskId InGraphExecutionTaskId, UPCGData* InData) { PerExecutionCache.SetExecutionCachePCGData(InGraphExecutionTaskId, EPCGPerExecutionCacheDataType::LandscapeData, InData); }
-	void SetLandscapeHeightPCGData(FPCGTaskId InGraphExecutionTaskId, UPCGData* InData) { PerExecutionCache.SetExecutionCachePCGData(InGraphExecutionTaskId, EPCGPerExecutionCacheDataType::LandscapeHeightData, InData); }
-	void SetOriginalSelfPCGData(FPCGTaskId InGraphExecutionTaskId, UPCGData* InData) { PerExecutionCache.SetExecutionCachePCGData(InGraphExecutionTaskId, EPCGPerExecutionCacheDataType::OriginalSelfData, InData); }
-	void SetBounds(FPCGTaskId InGraphExecutionTaskId, const FBox& InBounds) { PerExecutionCache.SetExecutionCacheBounds(InGraphExecutionTaskId, EPCGPerExecutionCacheDataType::Bounds, InBounds); }
-	void SetOriginalBounds(FPCGTaskId InGraphExecutionTaskId, const FBox& InBounds) { PerExecutionCache.SetExecutionCacheBounds(InGraphExecutionTaskId, EPCGPerExecutionCacheDataType::OriginalBounds, InBounds); }
-	void SetLocalSpaceBounds(FPCGTaskId InGraphExecutionTaskId, const FBox& InBounds) { PerExecutionCache.SetExecutionCacheBounds(InGraphExecutionTaskId, EPCGPerExecutionCacheDataType::LocalSpaceBounds, InBounds); }
-	void SetOriginalLocalSpaceBounds(FPCGTaskId InGraphExecutionTaskId, const FBox& InBounds) { PerExecutionCache.SetExecutionCacheBounds(InGraphExecutionTaskId, EPCGPerExecutionCacheDataType::OriginalLocalSpaceBounds, InBounds); }
-	void SetTotalBounds(FPCGTaskId InGraphExecutionTaskId, const FBox& InBounds) { PerExecutionCache.SetExecutionCacheBounds(InGraphExecutionTaskId, EPCGPerExecutionCacheDataType::TotalBounds, InBounds); }
-
 	double GetTickBudgetInSeconds() const;
 	int32 GetMaxNumProcessingTasks() const;
 	float GetMaxPercentageOfThreadsToUse() const;
diff --git a/Engine/Plugins/PCG/Source/PCG/Private/Graph/PCGGraphPerExecutionCache.cpp b/Engine/Plugins/PCG/Source/PCG/Private/Graph/PCGGraphPerExecutionCache.cpp
--- a/Engine/Plugins/PCG/Source/PCG/Private/Graph/PCGGraphPerExecutionCache.cpp
+++ b/Engine/Plugins/PCG/Source/PCG/Private/Graph/PCGGraphPerExecutionCache.cpp
@@ -1,68 +1,10 @@
 // Copyright Epic Games, Inc. All Rights Reserved.
 
-#include "PCGGraphPerExecutionCache.h"
+#include "Graph/PCGGraphPerExecutionCache.h"
 
-UPCGData* FPCGPerExecutionCache::GetExecutionCachePCGData(FPCGTaskId InGraphExecutionTaskId, EPCGPerExecutionCacheDataType InExecutionCacheDataType, bool& bOutFound)
-{
-	bOutFound = false;
-	if (InGraphExecutionTaskId != InvalidPCGTaskId)
-	{
-		PCG::TSharedScopeLock ReadLock(Lock);
-		if (FPCGPerExecutionCacheEntry* PerExecutionCacheEntry = Entries.Find(InGraphExecutionTaskId))
-		{
-			TInstancedStruct<FPCGPerExecutionCacheData>& StructData = PerExecutionCacheEntry->Data[(uint8)InExecutionCacheDataType];
-			if (StructData.IsValid())
-			{
-				bOutFound = true;
-				FPCGPerExecutionCachePCGData& PCGData = StructData.GetMutable<FPCGPerExecutionCachePCGData>();
-				return PCGData.Data;
-			}
-		}
-	}
-
-	return nullptr;
-}
-
-void FPCGPerExecutionCache::SetExecutionCachePCGData(FPCGTaskId InGraphExecutionTaskId, EPCGPerExecutionCacheDataType InExecutionCacheDataType, UPCGData* InData)
-{
-	if (InGraphExecutionTaskId != InvalidPCGTaskId)
-	{
-		PCG::TUniqueScopeLock WriteLock(Lock);
-		FPCGPerExecutionCacheEntry& PerExecutionCacheEntry = Entries.FindOrAdd(InGraphExecutionTaskId);
-		check(!PerExecutionCacheEntry.Data[(uint8)InExecutionCacheDataType].IsValid());
-		PerExecutionCacheEntry.Data[(uint8)InExecutionCacheDataType].InitializeAs<FPCGPerExecutionCachePCGData>(InData);
-	}
-}
-
-TOptional<FBox> FPCGPerExecutionCache::GetExecutionCacheBounds(FPCGTaskId InGraphExecutionTaskId, EPCGPerExecutionCacheDataType InExecutionCacheDataType)
-{
-	if (InGraphExecutionTaskId != InvalidPCGTaskId)
-	{
-		PCG::TSharedScopeLock ReadLock(Lock);
-		if (FPCGPerExecutionCacheEntry* PerExecutionCacheEntry = Entries.Find(InGraphExecutionTaskId))
-		{
-			TInstancedStruct<FPCGPerExecutionCacheData>& StructData = PerExecutionCacheEntry->Data[(uint8)InExecutionCacheDataType];
-			if (StructData.IsValid())
-			{
-				FPCGPerExecutionCacheBounds& BoundsData = StructData.GetMutable<FPCGPerExecutionCacheBounds>();
-				return BoundsData.Bounds;
-			}
-		}
-	}
-
-	return TOptional<FBox>();
-}
-
-void FPCGPerExecutionCache::SetExecutionCacheBounds(FPCGTaskId InGraphExecutionTaskId, EPCGPerExecutionCacheDataType InExecutionCacheDataType, const FBox& InBounds)
-{
-	if (InGraphExecutionTaskId != InvalidPCGTaskId)
-	{
-		PCG::TUniqueScopeLock WriteLock(Lock);
-		FPCGPerExecutionCacheEntry& PerExecutionCacheEntry = Entries.FindOrAdd(InGraphExecutionTaskId);
-		check(!PerExecutionCacheEntry.Data[(uint8)InExecutionCacheDataType].IsValid());
-		PerExecutionCacheEntry.Data[(uint8)InExecutionCacheDataType].InitializeAs<FPCGPerExecutionCacheBounds>(InBounds);
-	}
-}
+#include "PCGCommon.h"
+#include "PCGModule.h"
+#include "Subsystems/IPCGBaseSubsystem.h"
 
 void FPCGPerExecutionCache::Clear()
 {
@@ -72,7 +14,7 @@ void FPCGPerExecutionCache::Clear()
 
 void FPCGPerExecutionCache::AddStructReferencedObjects(FReferenceCollector& Collector)
 {
-	PCG::TSharedScopeLock ReadLock(Lock);
+	PCG::TSharedScopeLock ReadLock(Lock, !Lock.IsLockedByCurrentThread());
 	for (TMap<FPCGTaskId, FPCGPerExecutionCacheEntry>::TIterator It = Entries.CreateIterator(); It; ++It)
 	{
 		It.Value().AddStructReferencedObjects(Collector);
@@ -81,11 +23,11 @@ void FPCGPerExecutionCache::AddStructReferencedObjects(FReferenceCollector& Coll
 
 void FPCGPerExecutionCacheEntry::AddStructReferencedObjects(FReferenceCollector& Collector)
 {
-	for (uint8 DataIndex = 0; DataIndex < (uint8)EPCGPerExecutionCacheDataType::Count; ++DataIndex)
+	for (auto& [Id, StructData] : Data)
 	{
-		if (Data[DataIndex].IsValid())
+		if (StructData.IsValid())
 		{
-			Data[DataIndex].GetMutable().AddStructReferencedObjects(Collector);
+			StructData.GetMutable().AddStructReferencedObjects(Collector);
 		}
 	}
 }
@@ -94,4 +36,3 @@ void FPCGPerExecutionCachePCGData::AddStructReferencedObjects(FReferenceCollecto
 {
 	Collector.AddReferencedObject(Data);
 }
-
diff --git a/Engine/Plugins/PCG/Source/PCG/Private/Graph/PCGGraphPerExecutionCache.h b/Engine/Plugins/PCG/Source/PCG/Private/Graph/PCGGraphPerExecutionCache.h
deleted file mode 100644
diff --git a/Engine/Plugins/PCG/Source/PCG/Private/PCGComponent.cpp b/Engine/Plugins/PCG/Source/PCG/Private/PCGComponent.cpp
--- a/Engine/Plugins/PCG/Source/PCG/Private/PCGComponent.cpp
+++ b/Engine/Plugins/PCG/Source/PCG/Private/PCGComponent.cpp
@@ -17,6 +17,7 @@
 #include "Data/PCGPointData.h"
 #include "Data/PCGSpatialData.h"
 #include "Data/PCGUnionData.h"
+#include "Graph/PCGGraphPerExecutionCache.h"
 #include "Graph/PCGStackContext.h"
 #include "Grid/PCGPartitionActor.h"
 #include "Helpers/PCGActorHelpers.h"
@@ -548,6 +549,7 @@ FPCGTaskId UPCGComponent::GenerateInternal(bool bForce, uint32 Grid, EPCGCompone
 		ClearPerPinGeneratedOutput();
 
 #if WITH_EDITOR
+		PCG_EXECUTION_CACHE_VALIDATION_CREATE_SCOPE(this);
 		// Notify Subsystem first
 		GetSubsystem()->OnPCGGraphStartGenerating(this);
 #endif // WITH_EDITOR
@@ -2942,118 +2944,71 @@ void UPCGComponent::OnRefresh(bool bForceRefresh)
 // This applies to GetActorPCGData/GetLandscapePCGData/GetLandscapeHeightData/GetInputPCGData/GetPCGData
 UPCGData* UPCGComponent::GetPCGData() const
 {
-	UPCGData* Data = nullptr;
-	if (UPCGSubsystem* PCGSubsystem = GetSubsystem())
-	{
-		bool bFound = false;
-		Data = PCGSubsystem->GetPCGData(CurrentGenerationTask, bFound);
-		if (!bFound)
+	// @todo_pcg: before removing bValidateWritable=false validate first that this method will always be called after a ExecutePreGraph call
+	return UPCGSubsystem::GetOrCreateExecutionCacheValue(this, PCGPerExecutionCacheGuids::PCGData,
+		[this]()
 		{
-			Data = CachedPCGData ? CachedPCGData.Get() : CreatePCGData();
-			CachedPCGData = Data;
-
-			PCGSubsystem->SetPCGData(CurrentGenerationTask, Data);
-			if (PCGSubsystem->IsGraphCacheDebuggingEnabled() && CurrentGenerationTask != InvalidPCGTaskId)
+			if (!CachedPCGData)
 			{
-				UE_LOG(LogPCG, Log, TEXT("         [%s] CACHE REFRESH CachedPCGData"), *GetOwner()->GetName());
+				CachedPCGData = CreatePCGData();
 			}
-		}
-	}
-
-	return Data;
+			return CachedPCGData;
+		}, /*bValidateWritable=*/false);
 }
 
 UPCGData* UPCGComponent::GetInputPCGData() const
 {
-	UPCGData* Data = nullptr;
-	if (UPCGSubsystem* PCGSubsystem = GetSubsystem())
-	{
-		bool bFound = false;
-		Data = PCGSubsystem->GetInputPCGData(CurrentGenerationTask, bFound);
-		if (!bFound)
+	// @todo_pcg: before removing bValidateWritable=false validate first that this method will always be called after a ExecutePreGraph call
+	return UPCGSubsystem::GetOrCreateExecutionCacheValue(this, PCGPerExecutionCacheGuids::InputData,
+		[this]()
 		{
-			Data = CachedInputData ? CachedInputData.Get() : CreateInputPCGData();
-			CachedInputData = Data;
-
-			PCGSubsystem->SetInputPCGData(CurrentGenerationTask, Data);
-			if (PCGSubsystem->IsGraphCacheDebuggingEnabled() && CurrentGenerationTask != InvalidPCGTaskId)
+			if (!CachedInputData)
 			{
-				UE_LOG(LogPCG, Log, TEXT("         [%s] CACHE REFRESH CachedInputData"), *GetOwner()->GetName());
+				CachedInputData = CreateInputPCGData();
 			}
-		}
-	}
-
-	return Data;
+			return CachedInputData;
+		}, /*bValidateWritable=*/false);
 }
 
 UPCGData* UPCGComponent::GetActorPCGData() const
 {
-	UPCGData* Data = nullptr;
-	if (UPCGSubsystem* PCGSubsystem = GetSubsystem())
-	{
-		bool bFound = false;
-		Data = PCGSubsystem->GetSelfPCGData(CurrentGenerationTask, bFound);
-		if (!bFound)
+	return UPCGSubsystem::GetOrCreateExecutionCacheValue(this, PCGPerExecutionCacheGuids::SelfData,
+		[this]()
 		{
-			PCG_EXECUTION_CACHE_VALIDATION_CHECK(this);
-			Data = CachedActorData && !IsLandscapeCachedDataDirty(CachedActorData) ? CachedActorData.Get() : CreateActorPCGData();
-			CachedActorData = Data;
-
-			PCGSubsystem->SetSelfPCGData(CurrentGenerationTask, Data);
-			if (PCGSubsystem->IsGraphCacheDebuggingEnabled() && CurrentGenerationTask != InvalidPCGTaskId)
+			if (!CachedActorData || IsLandscapeCachedDataDirty(CachedActorData))
 			{
-				UE_LOG(LogPCG, Log, TEXT("         [%s] CACHE REFRESH CachedActorData"), *GetOwner()->GetName());
+				CachedActorData = CreateActorPCGData();
 			}
-		}
-	}
-
-	return Data;
+			return CachedActorData;
+		});
 }
 
 UPCGData* UPCGComponent::GetLandscapePCGData() const
 {
-	UPCGData* Data = nullptr;
-	if (UPCGSubsystem* PCGSubsystem = GetSubsystem())
-	{
-		bool bFound = false;
-		Data = PCGSubsystem->GetLandscapePCGData(CurrentGenerationTask, bFound);
-		if (!bFound)
+	// @todo_pcg: before removing bValidateWritable=false validate first that this method will always be called after a ExecutePreGraph call
+	return UPCGSubsystem::GetOrCreateExecutionCacheValue(this, PCGPerExecutionCacheGuids::LandscapeData,
+		[this]()
 		{
-			Data = CachedLandscapeData && !IsLandscapeCachedDataDirty(CachedLandscapeData) ? CachedLandscapeData.Get() : CreateLandscapePCGData(/*bHeightOnly=*/false);
-			CachedLandscapeData = Data;
-
-			PCGSubsystem->SetLandscapePCGData(CurrentGenerationTask, Data);
-			if (PCGSubsystem->IsGraphCacheDebuggingEnabled() && CurrentGenerationTask != InvalidPCGTaskId)
+			if (!CachedLandscapeData || IsLandscapeCachedDataDirty(CachedLandscapeData))
 			{
-				UE_LOG(LogPCG, Log, TEXT("         [%s] CACHE REFRESH CachedLandscapeData"), *GetOwner()->GetName());
+				CachedLandscapeData = CreateLandscapePCGData(/*bHeightOnly=*/false);
 			}
-		}
-	}
-
-	return Data;
+			return CachedLandscapeData;
+		}, /*bValidateWritable=*/false);
 }
 
 UPCGData* UPCGComponent::GetLandscapeHeightPCGData() const
 {
-	UPCGData* Data = nullptr;
-	if (UPCGSubsystem* PCGSubsystem = GetSubsystem())
-	{
-		bool bFound = false;
-		Data = PCGSubsystem->GetLandscapeHeightPCGData(CurrentGenerationTask, bFound);
-		if (!bFound)
+	// @todo_pcg: before removing bValidateWritable=false validate first that this method will always be called after a ExecutePreGraph call
+	return UPCGSubsystem::GetOrCreateExecutionCacheValue(this, PCGPerExecutionCacheGuids::LandscapeHeightData,
+		[this]()
 		{
-			Data = CachedLandscapeHeightData && !IsLandscapeCachedDataDirty(CachedLandscapeHeightData) ? CachedLandscapeHeightData.Get() : CreateLandscapePCGData(/*bHeightOnly=*/true);
-			CachedLandscapeHeightData = Data;
-
-			PCGSubsystem->SetLandscapeHeightPCGData(CurrentGenerationTask, Data);
-			if (PCGSubsystem->IsGraphCacheDebuggingEnabled() && CurrentGenerationTask != InvalidPCGTaskId)
+			if (!CachedLandscapeHeightData || IsLandscapeCachedDataDirty(CachedLandscapeHeightData))
 			{
-				UE_LOG(LogPCG, Log, TEXT("         [%s] CACHE REFRESH CachedLandscapeHeightData"), *GetOwner()->GetName());
+				CachedLandscapeHeightData = CreateLandscapePCGData(/*bHeightOnly=*/true);
 			}
-		}
-	}
-
-	return Data;
+			return CachedLandscapeHeightData;
+		}, /*bValidateWritable=*/false);
 }
 
 UPCGData* UPCGComponent::GetOriginalActorPCGData() const
@@ -3064,13 +3019,11 @@ UPCGData* UPCGComponent::GetOriginalActorPCGData() const
 		return GetActorPCGData();
 	}
 
-	UPCGData* Data = nullptr;
-	if (UPCGSubsystem* PCGSubsystem = GetSubsystem())
-	{
-		bool bFound = false;
-		Data = PCGSubsystem->GetOriginalSelfPCGData(CurrentGenerationTask, bFound);
-		if (!bFound)
+	// Special case where we forward to original so do our own writable validation
+	return UPCGSubsystem::GetOrCreateExecutionCacheValue(this, PCGPerExecutionCacheGuids::OriginalSelfData,
+		[this, PartitionActor]()
 		{
+			UPCGData* Data = nullptr;
 			if (UPCGComponent* OriginalComponent = PartitionActor->GetOriginalComponent(this))
 			{
 				Data = OriginalComponent->GetActorPCGData();
@@ -3079,16 +3032,9 @@ UPCGData* UPCGComponent::GetOriginalActorPCGData() const
 					PCG_EXECUTION_CACHE_VALIDATION_CHECK(this);
 					Data = OriginalComponent->CreateActorPCGData();
 				}
-				PCGSubsystem->SetOriginalSelfPCGData(CurrentGenerationTask, Data);
-				if (PCGSubsystem->IsGraphCacheDebuggingEnabled() && CurrentGenerationTask != InvalidPCGTaskId)
-				{
-					UE_LOG(LogPCG, Log, TEXT("         [%s] CACHE REFRESH CachedOriginalActorData"), *GetOwner()->GetName());
-				}
 			}
-		}
-	}
-
-	return Data;
+			return Data;
+		}, /*bValidateWritable=*/false);
 }
 
 UPCGComponent* UPCGComponent::GetOriginalComponent() const
@@ -3209,18 +3155,14 @@ void UPCGComponent::ExecutePreGraph(FPCGContext* InContext)
 
 	ensureMsgf(GetGenerationTaskId() != InvalidPCGTaskId, TEXT("Component was Scheduled for generation without having its CurrentGenerationTask assigned"));
 
-	{
-		PCG_EXECUTION_CACHE_VALIDATION_CREATE_SCOPE(Component);
-
-		// Call getters which will create the data and cache it
-		GetActorPCGData();
-		GetOriginalActorPCGData();
-		GetGridBounds();
-		GetOriginalGridBounds();
-		GetLocalSpaceBounds();
-		GetOriginalLocalSpaceBounds();
-		GetTotalBounds();
-	}
+	// Call getters which will create the data and cache it
+	GetActorPCGData();
+	GetOriginalActorPCGData();
+	GetGridBounds();
+	GetOriginalGridBounds();
+	GetLocalSpaceBounds();
+	GetOriginalLocalSpaceBounds();
+	GetTotalBounds();
 }
 
 UPCGData* UPCGComponent::CreatePCGData() const
@@ -3365,32 +3307,11 @@ bool UPCGComponent::IsLandscapeCachedDataDirty(const UPCGData* Data) const
 
 FBox UPCGComponent::GetGridBounds() const
 {
-	if (CurrentGenerationTask != InvalidPCGTaskId)
-	{
-		UPCGSubsystem* PCGSubsystem = GetSubsystem();
-		check(PCGSubsystem);
-		
-		TOptional<FBox> GridBounds = PCGSubsystem->GetBounds(CurrentGenerationTask);
-		
-		if (!GridBounds.IsSet())
+	return UPCGSubsystem::GetOrCreateExecutionCacheValue(this, PCGPerExecutionCacheGuids::Bounds,
+		[this]()
 		{
-			PCG_EXECUTION_CACHE_VALIDATION_CHECK(this);
-			GridBounds = PCGHelpers::GetGridBounds(GetOwner(), this);
-
-			PCGSubsystem->SetBounds(CurrentGenerationTask, GridBounds.GetValue());
-			if (PCGSubsystem->IsGraphCacheDebuggingEnabled())
-			{
-				UE_LOG(LogPCG, Log, TEXT("         [%s] CACHE REFRESH GridBounds"), *GetOwner()->GetName());
-			}
-		}
-
-		check(GridBounds.IsSet());
-		return GridBounds.GetValue();
-	}
-	else
-	{
-		return PCGHelpers::GetGridBounds(GetOwner(), this);
-	}
+			return PCGHelpers::GetGridBounds(GetOwner(), this);
+		});
 }
 
 FBox UPCGComponent::GetOriginalGridBounds() const
@@ -3401,67 +3322,22 @@ FBox UPCGComponent::GetOriginalGridBounds() const
 		return GetGridBounds();
 	}
 
-	if (CurrentGenerationTask != InvalidPCGTaskId)
-	{
-		UPCGSubsystem* PCGSubsystem = GetSubsystem();
-		check(PCGSubsystem);
-
-		TOptional<FBox> OriginalBounds = PCGSubsystem->GetOriginalBounds(CurrentGenerationTask);
-
-		if (!OriginalBounds.IsSet())
+	return UPCGSubsystem::GetOrCreateExecutionCacheValue(this, PCGPerExecutionCacheGuids::OriginalBounds,
+		[this, PartitionActor]()
 		{
-			PCG_EXECUTION_CACHE_VALIDATION_CHECK(this);
 			const UPCGComponent* OriginalComponent = PartitionActor->GetOriginalComponent(this);
 			const UPCGComponent* ValidComponent = OriginalComponent ? OriginalComponent : this;
-
-			OriginalBounds = PCGHelpers::GetGridBounds(ValidComponent->GetOwner(), ValidComponent);
-
-			PCGSubsystem->SetOriginalBounds(CurrentGenerationTask, OriginalBounds.GetValue());
-			if (PCGSubsystem->IsGraphCacheDebuggingEnabled())
-			{
-				UE_LOG(LogPCG, Log, TEXT("         [%s] CACHE REFRESH Original Bounds"), *GetOwner()->GetName());
-			}
-		}
-
-		check(OriginalBounds.IsSet());
-		return OriginalBounds.GetValue();
-	}
-	else
-	{
-		const UPCGComponent* OriginalComponent = PartitionActor->GetOriginalComponent(this);
-		const UPCGComponent* ValidComponent = OriginalComponent ? OriginalComponent : this;
-		return PCGHelpers::GetGridBounds(ValidComponent->GetOwner(), ValidComponent);
-	}
+			return PCGHelpers::GetGridBounds(ValidComponent->GetOwner(), ValidComponent);
+		});
 }
 
 FBox UPCGComponent::GetLocalSpaceBounds() const
 {
-	if (CurrentGenerationTask != InvalidPCGTaskId)
-	{
-		UPCGSubsystem* PCGSubsystem = GetSubsystem();
-		check(PCGSubsystem);
-
-		TOptional<FBox> LocalSpaceBounds = PCGSubsystem->GetLocalSpaceBounds(CurrentGenerationTask);
-
-		if (!LocalSpaceBounds.IsSet())
+	return UPCGSubsystem::GetOrCreateExecutionCacheValue(this, PCGPerExecutionCacheGuids::LocalSpaceBounds,
+		[this]()
 		{
-			PCG_EXECUTION_CACHE_VALIDATION_CHECK(this);
-			LocalSpaceBounds = PCGHelpers::GetActorLocalBounds(GetOwner());
-
-			PCGSubsystem->SetLocalSpaceBounds(CurrentGenerationTask, LocalSpaceBounds.GetValue());
-			if (PCGSubsystem->IsGraphCacheDebuggingEnabled())
-			{
-				UE_LOG(LogPCG, Log, TEXT("         [%s] CACHE REFRESH Local Space Bounds"), *GetOwner()->GetName());
-			}
-		}
-
-		check(LocalSpaceBounds.IsSet());
-		return LocalSpaceBounds.GetValue();
-	}
-	else
-	{
-		return PCGHelpers::GetActorLocalBounds(GetOwner());
-	}
+			return PCGHelpers::GetActorLocalBounds(GetOwner());
+		});
 }
 
 FBox UPCGComponent::GetOriginalLocalSpaceBounds() const
@@ -3472,68 +3348,23 @@ FBox UPCGComponent::GetOriginalLocalSpaceBounds() const
 		return GetLocalSpaceBounds();
 	}
 
-	if (CurrentGenerationTask != InvalidPCGTaskId)
-	{
-		UPCGSubsystem* PCGSubsystem = GetSubsystem();
-		check(PCGSubsystem);
-
-		TOptional<FBox> OriginalLocalSpaceBounds = PCGSubsystem->GetOriginalLocalSpaceBounds(CurrentGenerationTask);
-
-		if (!OriginalLocalSpaceBounds.IsSet())
+	return UPCGSubsystem::GetOrCreateExecutionCacheValue(this, PCGPerExecutionCacheGuids::OriginalLocalSpaceBounds,
+		[this, PartitionActor]()
 		{
-			PCG_EXECUTION_CACHE_VALIDATION_CHECK(this);
 			const UPCGComponent* OriginalComponent = PartitionActor->GetOriginalComponent(this);
 			const UPCGComponent* ValidComponent = OriginalComponent ? OriginalComponent : this;
-
-			OriginalLocalSpaceBounds = PCGHelpers::GetActorLocalBounds(ValidComponent->GetOwner());
-
-			PCGSubsystem->SetOriginalLocalSpaceBounds(CurrentGenerationTask, OriginalLocalSpaceBounds.GetValue());
-			if (PCGSubsystem->IsGraphCacheDebuggingEnabled())
-			{
-				UE_LOG(LogPCG, Log, TEXT("         [%s] CACHE REFRESH Original Local Space Bounds"), *GetOwner()->GetName());
-			}
-		}
-
-		check(OriginalLocalSpaceBounds.IsSet());
-		return OriginalLocalSpaceBounds.GetValue();
-	}
-	else
-	{
-		const UPCGComponent* OriginalComponent = PartitionActor->GetOriginalComponent(this);
-		const UPCGComponent* ValidComponent = OriginalComponent ? OriginalComponent : this;
-		return PCGHelpers::GetActorLocalBounds(ValidComponent->GetOwner());
-	}
+			return PCGHelpers::GetActorLocalBounds(ValidComponent->GetOwner());
+		});
 }
 
 FBox UPCGComponent::GetTotalBounds() const
 {
 	/* Total bounds are cached at the beginning of execution through UPCGComponent::ExecutePreGraph for tracking purposes of the previously generated resources. */
-	if (CurrentGenerationTask != InvalidPCGTaskId)
-	{
-		UPCGSubsystem* PCGSubsystem = GetSubsystem();
-		check(PCGSubsystem);
-
-		TOptional<FBox> TotalBounds = PCGSubsystem->GetTotalBounds(CurrentGenerationTask);
-
-		if (!TotalBounds.IsSet())
+	return UPCGSubsystem::GetOrCreateExecutionCacheValue(this, PCGPerExecutionCacheGuids::TotalBounds,
+		[this]()
 		{
-			PCG_EXECUTION_CACHE_VALIDATION_CHECK(this);
-			TotalBounds = PCGHelpers::GetActorBounds(GetOwner(), /*bIgnorePCGCreatedComponents=*/false);
-
-			PCGSubsystem->SetTotalBounds(CurrentGenerationTask, TotalBounds.GetValue());
-			if (PCGSubsystem->IsGraphCacheDebuggingEnabled())
-			{
-				UE_LOG(LogPCG, Log, TEXT("         [%s] CACHE REFRESH Total Bounds"), *GetOwner()->GetName());
-			}
-		}
-
-		check(TotalBounds.IsSet());
-		return TotalBounds.GetValue();
-	}
-	else
-	{
-		return PCGHelpers::GetActorBounds(GetOwner(), /*bIgnorePCGCreatedComponents=*/false);
-	}
+			return PCGHelpers::GetActorBounds(GetOwner(), /*bIgnorePCGCreatedComponents=*/false);
+		});
 }
 
 FBox UPCGComponent::GetGridBounds(const AActor* Actor) const
diff --git a/Engine/Plugins/PCG/Source/PCG/Private/PCGDefaultActorExecutionSource.cpp b/Engine/Plugins/PCG/Source/PCG/Private/PCGDefaultActorExecutionSource.cpp
--- a/Engine/Plugins/PCG/Source/PCG/Private/PCGDefaultActorExecutionSource.cpp
+++ b/Engine/Plugins/PCG/Source/PCG/Private/PCGDefaultActorExecutionSource.cpp
@@ -4,6 +4,7 @@
 
 #include "PCGSubsystem.h"
 #include "Data/PCGUnionData.h"
+#include "Graph/PCGGraphPerExecutionCache.h"
 #include "Helpers/PCGHelpers.h"
 
 #include "GameFramework/Actor.h"
@@ -32,58 +33,51 @@ FString FPCGDefaultActorExecutionState::GetDebugName() const
 	}
 }
 
-UPCGData* FPCGDefaultActorExecutionState::GetSelfData() const
+
+UPCGData* FPCGDefaultActorExecutionState::GetSelfDataInternal() const
 {
-	if (AActor* Actor = GetActor())
+	AActor* Actor = GetActor();
+	if (!Actor)
+	{
+		return FPCGDefaultWorldObjectExecutionState::GetSelfData();
+	}
+	
+	FPCGGetDataFunctionRegistryParams Params;
+	Params.Source = GetSource();
+	Params.bParseActor = true;
+	Params.DataTypeFilter = EPCGDataType::Spatial;
+
+	FPCGGetDataFunctionRegistryOutput Output;
+	FPCGModule::ConstGetDataFunctionRegistry().GetDataFromActor(/*Context=*/nullptr, Params, Actor, Output);
+
+	if (Output.Collection.TaggedData.Num() > 1)
 	{
-		if (UPCGSubsystem* PCGSubsystem = UPCGSubsystem::GetInstance(Actor->GetWorld()))
+		UPCGUnionData* Union = NewObject<UPCGUnionData>();
+		for (const FPCGTaggedData& TaggedData : Output.Collection.TaggedData)
 		{
-			bool bFound = false;
-			return PCGSubsystem->GetSelfPCGData(GetTypedSource()->CurrentGenerationTask, bFound);
+			Union->AddData(CastChecked<const UPCGSpatialData>(TaggedData.Data));
 		}
+
+		return Union;
+	}
+	else if (Output.Collection.TaggedData.Num() == 1)
+	{
+		return const_cast<UPCGData*>(Output.Collection.TaggedData[0].Data.Get());
 	}
 
-	return FPCGDefaultWorldObjectExecutionState::GetSelfData();
+	return nullptr;
+}
+
+UPCGData* FPCGDefaultActorExecutionState::GetSelfData() const
+{
+	return UPCGSubsystem::GetOrCreateExecutionCacheValue(GetTypedSource(), PCGPerExecutionCacheGuids::SelfData, [this]() { return GetSelfDataInternal(); });
 }
 
 void FPCGDefaultActorExecutionState::ExecutePreGraph(FPCGContext* Context)
 {
-	// Actor execution states have only self-data, and will not fill in other pins.
-	// @todo_pcg: should we also set bounds?
-	AActor* Actor = GetActor();
-	UPCGSubsystem* PCGSubsystem = Actor ? UPCGSubsystem::GetInstance(Actor->GetWorld()) : nullptr;
-	if (PCGSubsystem)
-	{
-		bool bFound = false;
-		UPCGData* SelfData = PCGSubsystem->GetSelfPCGData(GetTypedSource()->CurrentGenerationTask, bFound);
-		if (!bFound)
-		{
-			FPCGGetDataFunctionRegistryParams Params;
-			Params.Source = GetSource();
-			Params.bParseActor = true;
-			Params.DataTypeFilter = EPCGDataType::Spatial;
-
-			FPCGGetDataFunctionRegistryOutput Output;
-			FPCGModule::ConstGetDataFunctionRegistry().GetDataFromActor(/*Context=*/nullptr, Params, Actor, Output);
-
-			if (Output.Collection.TaggedData.Num() > 1)
-			{
-				UPCGUnionData* Union = NewObject<UPCGUnionData>();
-				for (const FPCGTaggedData& TaggedData : Output.Collection.TaggedData)
-				{
-					Union->AddData(CastChecked<const UPCGSpatialData>(TaggedData.Data));
-				}
-
-				SelfData = Union;
-			}
-			else if (Output.Collection.TaggedData.Num() == 1)
-			{
-				SelfData = const_cast<UPCGData*>(Output.Collection.TaggedData[0].Data.Get());
-			}
-
-			PCGSubsystem->SetSelfPCGData(GetTypedSource()->CurrentGenerationTask, SelfData);
-		}
-	}
+	GetSelfData();
+	GetBounds();
+	GetLocalSpaceBounds();
 }
 
 FTransform FPCGDefaultActorExecutionState::GetTransform() const
@@ -115,7 +109,7 @@ bool FPCGDefaultActorExecutionState::HasAuthority() const
 	}
 }
 
-FBox FPCGDefaultActorExecutionState::GetBounds() const
+FBox FPCGDefaultActorExecutionState::GetBoundsInternal() const
 {
 	if (AActor* Actor = GetActor())
 	{
@@ -127,7 +121,12 @@ FBox FPCGDefaultActorExecutionState::GetBounds() const
 	}
 }
 
-FBox FPCGDefaultActorExecutionState::GetLocalSpaceBounds() const
+FBox FPCGDefaultActorExecutionState::GetBounds() const
+{
+	return UPCGSubsystem::GetOrCreateExecutionCacheValue(GetTypedSource(), PCGPerExecutionCacheGuids::Bounds, [this]() { return GetBoundsInternal(); });
+}
+
+FBox FPCGDefaultActorExecutionState::GetLocalSpaceBoundsInternal() const
 {
 	if (AActor* Actor = GetActor())
 	{
@@ -139,6 +138,11 @@ FBox FPCGDefaultActorExecutionState::GetLocalSpaceBounds() const
 	}
 }
 
+FBox FPCGDefaultActorExecutionState::GetLocalSpaceBounds() const
+{
+	return UPCGSubsystem::GetOrCreateExecutionCacheValue(GetTypedSource(), PCGPerExecutionCacheGuids::LocalSpaceBounds, [this]() { return GetLocalSpaceBoundsInternal(); });
+}
+
 UPCGDefaultActorExecutionSource::UPCGDefaultActorExecutionSource()
 {
 	// Inefficient but more foolproof - we'll replace the previously created state from the base class by our version.
diff --git a/Engine/Plugins/PCG/Source/PCG/Private/RuntimeGen/PCGRuntimeGenExecutionSource.cpp b/Engine/Plugins/PCG/Source/PCG/Private/RuntimeGen/PCGRuntimeGenExecutionSource.cpp
--- a/Engine/Plugins/PCG/Source/PCG/Private/RuntimeGen/PCGRuntimeGenExecutionSource.cpp
+++ b/Engine/Plugins/PCG/Source/PCG/Private/RuntimeGen/PCGRuntimeGenExecutionSource.cpp
@@ -8,6 +8,7 @@
 #include "PCGManagedResourceContainer.h"
 #include "PCGSubsystem.h"
 #include "Data/PCGVolumeData.h"
+#include "Graph/PCGGraphPerExecutionCache.h"
 #include "Utils/PCGGeneratedResourcesLogging.h"
 #include "Utils/PCGGraphExecutionLogging.h"
 
@@ -174,26 +175,7 @@ FBox UPCGRuntimeGenExecutionSource::GetBoundsInternal() const
 
 FBox UPCGRuntimeGenExecutionSource::GetBounds() const
 {
-	if (GenerationTaskId != InvalidPCGTaskId)
-	{
-		UPCGSubsystem* PCGSubsystem = Cast<UPCGSubsystem>(ExecutionState.GetSubsystem());
-		check(PCGSubsystem);
-
-		TOptional<FBox> GridBounds = PCGSubsystem->GetBounds(GenerationTaskId);
-
-		if (!GridBounds.IsSet())
-		{
-			GridBounds = GetBoundsInternal();
-			PCGSubsystem->SetBounds(GenerationTaskId, GridBounds.GetValue());
-		}
-
-		check(GridBounds.IsSet());
-		return GridBounds.GetValue();
-	}
-	else
-	{
-		return GetBoundsInternal();
-	}
+	return UPCGSubsystem::GetOrCreateExecutionCacheValue(this, PCGPerExecutionCacheGuids::Bounds, [this]() { return GetBoundsInternal(); });
 }
 
 FBox UPCGRuntimeGenExecutionSource::GetOriginalBoundsInternal() const
@@ -204,26 +186,7 @@ FBox UPCGRuntimeGenExecutionSource::GetOriginalBoundsInternal() const
 
 FBox UPCGRuntimeGenExecutionSource::GetOriginalBounds() const
 {
-	if (GenerationTaskId != InvalidPCGTaskId)
-	{
-		UPCGSubsystem* PCGSubsystem = Cast<UPCGSubsystem>(ExecutionState.GetSubsystem());
-		check(PCGSubsystem);
-
-		TOptional<FBox> GridBounds = PCGSubsystem->GetOriginalBounds(GenerationTaskId);
-
-		if (!GridBounds.IsSet())
-		{
-			GridBounds = GetOriginalBoundsInternal();
-			PCGSubsystem->SetOriginalBounds(GenerationTaskId, GridBounds.GetValue());
-		}
-
-		check(GridBounds.IsSet());
-		return GridBounds.GetValue();
-	}
-	else
-	{
-		return GetOriginalBoundsInternal();
-	}
+	return UPCGSubsystem::GetOrCreateExecutionCacheValue(this, PCGPerExecutionCacheGuids::OriginalBounds, [this]() { return GetOriginalBoundsInternal(); });
 }
 
 FBox UPCGRuntimeGenExecutionSource::GetLocalSpaceBoundsInternal() const
@@ -241,26 +204,7 @@ FBox UPCGRuntimeGenExecutionSource::GetLocalSpaceBoundsInternal() const
 
 FBox UPCGRuntimeGenExecutionSource::GetLocalSpaceBounds() const
 {
-	if (GenerationTaskId != InvalidPCGTaskId)
-	{
-		UPCGSubsystem* PCGSubsystem = Cast<UPCGSubsystem>(ExecutionState.GetSubsystem());
-		check(PCGSubsystem);
-
-		TOptional<FBox> GridBounds = PCGSubsystem->GetLocalSpaceBounds(GenerationTaskId);
-
-		if (!GridBounds.IsSet())
-		{
-			GridBounds = GetLocalSpaceBoundsInternal();
-			PCGSubsystem->SetLocalSpaceBounds(GenerationTaskId, GridBounds.GetValue());
-		}
-
-		check(GridBounds.IsSet());
-		return GridBounds.GetValue();
-	}
-	else
-	{
-		return GetLocalSpaceBoundsInternal();
-	}
+	return UPCGSubsystem::GetOrCreateExecutionCacheValue(this, PCGPerExecutionCacheGuids::LocalSpaceBounds, [this]() { return GetLocalSpaceBoundsInternal(); });
 }
 
 FBox UPCGRuntimeGenExecutionSource::GetOriginalLocalSpaceBoundsInternal() const
@@ -271,26 +215,7 @@ FBox UPCGRuntimeGenExecutionSource::GetOriginalLocalSpaceBoundsInternal() const
 
 FBox UPCGRuntimeGenExecutionSource::GetOriginalLocalSpaceBounds() const
 {
-	if (GenerationTaskId != InvalidPCGTaskId)
-	{
-		UPCGSubsystem* PCGSubsystem = Cast<UPCGSubsystem>(ExecutionState.GetSubsystem());
-		check(PCGSubsystem);
-
-		TOptional<FBox> GridBounds = PCGSubsystem->GetOriginalLocalSpaceBounds(GenerationTaskId);
-
-		if (!GridBounds.IsSet())
-		{
-			GridBounds = GetOriginalLocalSpaceBoundsInternal();
-			PCGSubsystem->SetOriginalLocalSpaceBounds(GenerationTaskId, GridBounds.GetValue());
-		}
-
-		check(GridBounds.IsSet());
-		return GridBounds.GetValue();
-	}
-	else
-	{
-		return GetOriginalLocalSpaceBoundsInternal();
-	}
+	return UPCGSubsystem::GetOrCreateExecutionCacheValue(this, PCGPerExecutionCacheGuids::OriginalLocalSpaceBounds, [this]() { return GetOriginalLocalSpaceBoundsInternal(); });
 }
 
 UPCGGraph* UPCGRuntimeGenExecutionSource::GetGraph() const
diff --git a/Engine/Plugins/PCG/Source/PCG/Private/Subsystems/IPCGBaseSubsystem.cpp b/Engine/Plugins/PCG/Source/PCG/Private/Subsystems/IPCGBaseSubsystem.cpp
--- a/Engine/Plugins/PCG/Source/PCG/Private/Subsystems/IPCGBaseSubsystem.cpp
+++ b/Engine/Plugins/PCG/Source/PCG/Private/Subsystems/IPCGBaseSubsystem.cpp
@@ -51,6 +51,11 @@ void IPCGBaseSubsystem::OnScheduleGraph(const FPCGStackContext& StackContext)
 }
 #endif
 
+FPCGPerExecutionCache* IPCGBaseSubsystem::GetCacheInternal() const
+{
+	return GraphExecutor ? &GraphExecutor->PerExecutionCache : nullptr;
+}
+
 FPCGTaskId IPCGBaseSubsystem::ScheduleGraph(const FPCGScheduleGraphParams& InParams)
 {
 	if (InParams.ExecutionSource)
diff --git a/Engine/Plugins/PCG/Source/PCG/Private/Subsystems/PCGSubsystem.cpp b/Engine/Plugins/PCG/Source/PCG/Private/Subsystems/PCGSubsystem.cpp
--- a/Engine/Plugins/PCG/Source/PCG/Private/Subsystems/PCGSubsystem.cpp
+++ b/Engine/Plugins/PCG/Source/PCG/Private/Subsystems/PCGSubsystem.cpp
@@ -2110,155 +2110,6 @@ void UPCGSubsystem::ResetPartitionActorsMap()
 
 #endif // WITH_EDITOR
 
-UPCGData* UPCGSubsystem::GetPCGData(FPCGTaskId InGraphExecutionTaskId, bool& bOutFound)
-{
-	bOutFound = false;
-	return GraphExecutor ? GraphExecutor->GetPCGData(InGraphExecutionTaskId, bOutFound) : nullptr;
-}
-
-UPCGData* UPCGSubsystem::GetInputPCGData(FPCGTaskId InGraphExecutionTaskId, bool& bOutFound)
-{
-	bOutFound = false;
-	return GraphExecutor ? GraphExecutor->GetInputPCGData(InGraphExecutionTaskId, bOutFound) : nullptr;
-}
-
-UPCGData* UPCGSubsystem::GetSelfPCGData(FPCGTaskId InGraphExecutionTaskId, bool& bOutFound)
-{
-	bOutFound = false;
-	return GraphExecutor ? GraphExecutor->GetSelfPCGData(InGraphExecutionTaskId, bOutFound) : nullptr;
-}
-
-UPCGData* UPCGSubsystem::GetLandscapePCGData(FPCGTaskId InGraphExecutionTaskId, bool& bOutFound)
-{
-	bOutFound = false;
-	return GraphExecutor ? GraphExecutor->GetLandscapePCGData(InGraphExecutionTaskId, bOutFound) : nullptr;
-}
-
-UPCGData* UPCGSubsystem::GetLandscapeHeightPCGData(FPCGTaskId InGraphExecutionTaskId, bool& bOutFound)
-{
-	bOutFound = false;
-	return GraphExecutor ? GraphExecutor->GetLandscapeHeightPCGData(InGraphExecutionTaskId, bOutFound) : nullptr;
-}
-
-UPCGData* UPCGSubsystem::GetOriginalSelfPCGData(FPCGTaskId InGraphExecutionTaskId, bool& bOutFound)
-{
-	bOutFound = false;
-	return GraphExecutor ? GraphExecutor->GetOriginalSelfPCGData(InGraphExecutionTaskId, bOutFound) : nullptr;
-}
-
-TOptional<FBox> UPCGSubsystem::GetBounds(FPCGTaskId InGraphExecutionTaskId)
-{
-	return GraphExecutor ? GraphExecutor->GetBounds(InGraphExecutionTaskId) : TOptional<FBox>();
-}
-
-TOptional<FBox> UPCGSubsystem::GetOriginalBounds(FPCGTaskId InGraphExecutionTaskId)
-{
-	return GraphExecutor ? GraphExecutor->GetOriginalBounds(InGraphExecutionTaskId) : TOptional<FBox>();
-}
-
-TOptional<FBox> UPCGSubsystem::GetLocalSpaceBounds(FPCGTaskId InGraphExecutionTaskId)
-{
-	return GraphExecutor ? GraphExecutor->GetLocalSpaceBounds(InGraphExecutionTaskId) : TOptional<FBox>();
-}
-
-TOptional<FBox> UPCGSubsystem::GetOriginalLocalSpaceBounds(FPCGTaskId InGraphExecutionTaskId)
-{
-	return GraphExecutor ? GraphExecutor->GetOriginalLocalSpaceBounds(InGraphExecutionTaskId) : TOptional<FBox>();
-}
-
-TOptional<FBox> UPCGSubsystem::GetTotalBounds(FPCGTaskId InGraphExecutionTaskId)
-{
-	return GraphExecutor ? GraphExecutor->GetTotalBounds(InGraphExecutionTaskId) : TOptional<FBox>();
-}
-
-void UPCGSubsystem::SetPCGData(FPCGTaskId InGraphExecutionTaskId, UPCGData* InData)
-{
-	if (GraphExecutor)
-	{
-		GraphExecutor->SetPCGData(InGraphExecutionTaskId, InData);
-	}
-}
-
-void UPCGSubsystem::SetInputPCGData(FPCGTaskId InGraphExecutionTaskId, UPCGData* InData)
-{
-	if (GraphExecutor)
-	{
-		GraphExecutor->SetInputPCGData(InGraphExecutionTaskId, InData);
-	}
-}
-
-void UPCGSubsystem::SetSelfPCGData(FPCGTaskId InGraphExecutionTaskId, UPCGData* InData)
-{
-	if (GraphExecutor)
-	{
-		GraphExecutor->SetSelfPCGData(InGraphExecutionTaskId, InData);
-	}
-}
-
-void UPCGSubsystem::SetLandscapePCGData(FPCGTaskId InGraphExecutionTaskId, UPCGData* InData)
-{
-	if (GraphExecutor)
-	{
-		GraphExecutor->SetLandscapePCGData(InGraphExecutionTaskId, InData);
-	}
-}
-
-void UPCGSubsystem::SetLandscapeHeightPCGData(FPCGTaskId InGraphExecutionTaskId, UPCGData* InData)
-{
-	if (GraphExecutor)
-	{
-		GraphExecutor->SetLandscapeHeightPCGData(InGraphExecutionTaskId, InData);
-	}
-}
-
-void UPCGSubsystem::SetOriginalSelfPCGData(FPCGTaskId InGraphExecutionTaskId, UPCGData* InData)
-{
-	if (GraphExecutor)
-	{
-		GraphExecutor->SetOriginalSelfPCGData(InGraphExecutionTaskId, InData);
-	}
-}
-
-void UPCGSubsystem::SetBounds(FPCGTaskId InGraphExecutionTaskId, const FBox& InBounds)
-{
-	if (GraphExecutor)
-	{
-		GraphExecutor->SetBounds(InGraphExecutionTaskId, InBounds);
-	}
-}
-
-void UPCGSubsystem::SetOriginalBounds(FPCGTaskId InGraphExecutionTaskId, const FBox& InBounds)
-{
-	if (GraphExecutor)
-	{
-		GraphExecutor->SetOriginalBounds(InGraphExecutionTaskId, InBounds);
-	}
-}
-
-void UPCGSubsystem::SetLocalSpaceBounds(FPCGTaskId InGraphExecutionTaskId, const FBox& InBounds)
-{
-	if (GraphExecutor)
-	{
-		GraphExecutor->SetLocalSpaceBounds(InGraphExecutionTaskId, InBounds);
-	}
-}
-
-void UPCGSubsystem::SetOriginalLocalSpaceBounds(FPCGTaskId InGraphExecutionTaskId, const FBox& InBounds)
-{
-	if (GraphExecutor)
-	{
-		GraphExecutor->SetOriginalLocalSpaceBounds(InGraphExecutionTaskId, InBounds);
-	}
-}
-
-void UPCGSubsystem::SetTotalBounds(FPCGTaskId InGraphExecutionTaskId, const FBox& InBounds)
-{
-	if (GraphExecutor)
-	{
-		GraphExecutor->SetTotalBounds(InGraphExecutionTaskId, InBounds);
-	}
-}
-
 void UPCGSubsystem::ExecuteBeginTickActions()
 {
 	TArray<FTickAction> Actions = MoveTemp(BeginTickActions);
diff --git a/Engine/Plugins/PCG/Source/PCG/Public/Graph/PCGGraphPerExecutionCache.h b/Engine/Plugins/PCG/Source/PCG/Public/Graph/PCGGraphPerExecutionCache.h
new file mode 100644
--- /dev/null
+++ b/Engine/Plugins/PCG/Source/PCG/Public/Graph/PCGGraphPerExecutionCache.h
@@ -0,0 +1,210 @@
+// Copyright Epic Games, Inc. All Rights Reserved.
+
+#pragma once
+
+#include "PCGCommon.h"
+#include "PCGModule.h"
+#include "PCGGraphExecutionStateInterface.h"
+
+#include "StructUtils/InstancedStruct.h"
+
+#include "PCGGraphPerExecutionCache.generated.h"
+
+class UPCGData;
+
+USTRUCT()
+struct FPCGPerExecutionCacheData
+{
+	GENERATED_BODY()
+
+	FPCGPerExecutionCacheData() = default;
+	virtual ~FPCGPerExecutionCacheData() = default;
+
+	virtual void AddStructReferencedObjects(FReferenceCollector& Collector) {}
+};
+
+USTRUCT()
+struct FPCGPerExecutionCachePCGData : public FPCGPerExecutionCacheData
+{
+	GENERATED_BODY()
+
+	using ValueType = UPCGData*;
+
+	FPCGPerExecutionCachePCGData() = default;
+	FPCGPerExecutionCachePCGData(ValueType InData)
+		: Data(InData)
+	{
+	}
+
+	virtual void AddStructReferencedObjects(FReferenceCollector& Collector) override;
+
+	ValueType GetValue() const { return Data; }
+
+	// stored as TObjectPtr for AddStructReferencedOjects
+	TObjectPtr<UPCGData> Data = nullptr;
+};
+
+USTRUCT()
+struct FPCGPerExecutionCacheBounds : public FPCGPerExecutionCacheData
+{
+	GENERATED_BODY()
+
+	using ValueType = FBox;
+
+	FPCGPerExecutionCacheBounds() = default;
+	FPCGPerExecutionCacheBounds(const ValueType& InBounds)
+		: Bounds(InBounds)
+	{
+	}
+
+	const ValueType& GetValue() const { return Bounds; }
+
+	ValueType Bounds;
+};
+
+/** Typed GUID wrapper associating a cache identifier with a concrete FPCGPerExecutionCacheData
+ *  subclass. Using TPCGPerExecutionCacheId<T> instead of a raw FGuid in Get/Set prevents passing a
+ *  Bounds GUID to a PCGData getter — the type mismatch is caught at compile time. */
+template<typename T>
+struct TPCGPerExecutionCacheId
+{
+	FGuid Guid;
+	const TCHAR* DebugName = nullptr;
+	explicit constexpr TPCGPerExecutionCacheId(FGuid InGuid, const TCHAR* InDebugName = nullptr) : Guid(InGuid), DebugName(InDebugName) {}
+};
+
+namespace PCGPerExecutionCacheGuids
+{
+	constexpr TPCGPerExecutionCacheId<FPCGPerExecutionCachePCGData>    PCGData{						FGuid(0x8D1077E7, 0x32AD4876, 0xA86D33D4, 0xE6618E25), TEXT("CachedPCGData") };
+	constexpr TPCGPerExecutionCacheId<FPCGPerExecutionCachePCGData>    InputData{					FGuid(0x410F0D9A, 0x10504945, 0x916FAEC0, 0xAC4CCDF5), TEXT("CachedInputData") };
+	constexpr TPCGPerExecutionCacheId<FPCGPerExecutionCachePCGData>    SelfData{					FGuid(0xDB4FEE15, 0x0B924BB9, 0x920E17D0, 0xB68A439A), TEXT("CachedActorData") };
+	constexpr TPCGPerExecutionCacheId<FPCGPerExecutionCachePCGData>    LandscapeData{				FGuid(0xBD2F8989, 0xF911433C, 0x9150CF07, 0xD25A2A3F), TEXT("CachedLandscapeData") };
+	constexpr TPCGPerExecutionCacheId<FPCGPerExecutionCachePCGData>    LandscapeHeightData{			FGuid(0x697872EE, 0x90624718, 0x85248924, 0xC5C3C2E2), TEXT("CachedLandscapeHeightData") };
+	constexpr TPCGPerExecutionCacheId<FPCGPerExecutionCachePCGData>    OriginalSelfData{			FGuid(0xE2A8E4C1, 0x18B444BD, 0xABA60F8C, 0x78FFC972), TEXT("CachedOriginalActorData") };
+	constexpr TPCGPerExecutionCacheId<FPCGPerExecutionCacheBounds>     Bounds{						FGuid(0xEEB1FE5D, 0xCA5B4EDA, 0xA96B1448, 0x49F7FFAE), TEXT("GridBounds") };
+	constexpr TPCGPerExecutionCacheId<FPCGPerExecutionCacheBounds>     OriginalBounds{				FGuid(0x7A77D800, 0x5E9045D3, 0xB69CD774, 0xE6CC5F0C), TEXT("Original Bounds") };
+	constexpr TPCGPerExecutionCacheId<FPCGPerExecutionCacheBounds>     LocalSpaceBounds{			FGuid(0xDA7B1698, 0x19A54A97, 0xB53C66E4, 0x126319FC), TEXT("Local Space Bounds") };
+	constexpr TPCGPerExecutionCacheId<FPCGPerExecutionCacheBounds>     OriginalLocalSpaceBounds{	FGuid(0x4F358809, 0x6FB24086, 0xB977FC0C, 0x8BC291D3), TEXT("Original Local Space Bounds") };
+	constexpr TPCGPerExecutionCacheId<FPCGPerExecutionCacheBounds>     TotalBounds{					FGuid(0x65B51D8B, 0xA70348BE, 0x894CDDF1, 0x8D318171), TEXT("Total Bounds") };
+}
+
+struct FPCGPerExecutionCacheEntry
+{
+	void AddStructReferencedObjects(FReferenceCollector& Collector);
+
+	TMap<FGuid, TInstancedStruct<FPCGPerExecutionCacheData>> Data;
+};
+
+struct FPCGPerExecutionCache
+{
+	/** Returns the cached value for the source's current task, unset otherwise.
+	 *  T is deduced from InCacheId — mismatched IDs are a compile error. */
+	template<typename T>
+	TOptional<typename T::ValueType> GetExecutionCacheEntry(const IPCGGraphExecutionSource* InSource, TPCGPerExecutionCacheId<T> InCacheId);
+
+	/** Stores a new value for the source's current task and typed ID. Asserts if the slot is already occupied. */
+	template<typename T>
+	void SetExecutionCacheEntry(const IPCGGraphExecutionSource* InSource, TPCGPerExecutionCacheId<T> InCacheId, typename T::ValueType InValue, bool bInValidateWritable = true, bool bEnableLogging = false);
+
+	/** Returns the cached value if present, otherwise calls InMakeEntry() to produce one.
+	 *  If the source's task is valid the produced value is stored in the cache before returning.
+	 *  InMakeEntry must return T::ValueType by value. */
+	template<typename T, typename FuncType>
+	typename T::ValueType GetOrCreateExecutionCacheValue(const IPCGGraphExecutionSource* InSource, TPCGPerExecutionCacheId<T> InCacheId, const FuncType& InMakeEntry, bool bInValidateWritable = true, bool bEnableLogging = false);
+
+	void Clear();
+	void AddStructReferencedObjects(FReferenceCollector& Collector);
+
+	/** Per graph execution cache, gets emptied when executor has no more work to do */
+	PCG::FRecursiveSharedLock Lock;
+	TMap<FPCGTaskId, FPCGPerExecutionCacheEntry> Entries;
+};
+
+template<typename T>
+TOptional<typename T::ValueType> FPCGPerExecutionCache::GetExecutionCacheEntry(const IPCGGraphExecutionSource* InSource, TPCGPerExecutionCacheId<T> InCacheId)
+{
+	check(InSource);
+	const FPCGTaskId TaskId = InSource->GetExecutionState().GetGenerationTaskId();
+	if (TaskId != InvalidPCGTaskId)
+	{
+		PCG::TSharedScopeLock ReadLock(Lock, !Lock.IsLockedByCurrentThread());
+		if (FPCGPerExecutionCacheEntry* Entry = Entries.Find(TaskId))
+		{
+			if (TInstancedStruct<FPCGPerExecutionCacheData>* StructData = Entry->Data.Find(InCacheId.Guid))
+			{
+				if (StructData->IsValid())
+				{
+					return StructData->Get<T>().GetValue();
+				}
+			}
+		}
+	}
+	return {};
+}
+
+template<typename T>
+void FPCGPerExecutionCache::SetExecutionCacheEntry(const IPCGGraphExecutionSource* InSource, TPCGPerExecutionCacheId<T> InCacheId, typename T::ValueType InValue, bool bValidateWritable, bool bEnableLogging)
+{
+	check(InSource);
+	const FPCGTaskId TaskId = InSource->GetExecutionState().GetGenerationTaskId();
+	if (TaskId != InvalidPCGTaskId)
+	{
+		PCG::TUniqueScopeLock WriteLock(Lock);
+		FPCGPerExecutionCacheEntry& Entry = Entries.FindOrAdd(TaskId);
+		TInstancedStruct<FPCGPerExecutionCacheData>& StructData = Entry.Data.FindOrAdd(InCacheId.Guid);
+		check(!StructData.IsValid());
+		
+		if (bValidateWritable)
+		{
+			PCG_EXECUTION_CACHE_VALIDATION_CHECK(InSource);
+		}
+
+		StructData.InitializeAs<T>(InValue);
+
+		UE_CLOG(bEnableLogging, LogPCG, Log, TEXT("         [%s] CACHE SET %s"),
+			*InSource->GetExecutionState().GetDebugName(),
+			InCacheId.DebugName ? InCacheId.DebugName : TEXT("unknown"));
+	}
+}
+
+template<typename T, typename FuncType>
+typename T::ValueType FPCGPerExecutionCache::GetOrCreateExecutionCacheValue(const IPCGGraphExecutionSource* InSource, TPCGPerExecutionCacheId<T> InCacheId, const FuncType& InMakeEntry, bool bValidateWritable, bool bEnableLogging)
+{
+	check(InSource);
+	const FPCGTaskId TaskId = InSource->GetExecutionState().GetGenerationTaskId();
+	if (TaskId != InvalidPCGTaskId)
+	{
+		{
+			PCG::TSharedScopeLock ReadLock(Lock, !Lock.IsLockedByCurrentThread());
+			if (FPCGPerExecutionCacheEntry* Entry = Entries.Find(TaskId))
+			{
+				if (TInstancedStruct<FPCGPerExecutionCacheData>* StructData = Entry->Data.Find(InCacheId.Guid))
+				{
+					if (StructData->IsValid())
+					{
+						return StructData->GetMutable<T>().GetValue();
+					}
+				}
+			}
+		}
+
+		PCG::TUniqueScopeLock WriteLock(Lock);
+		FPCGPerExecutionCacheEntry& Entry = Entries.FindOrAdd(TaskId);
+		TInstancedStruct<FPCGPerExecutionCacheData>& StructData = Entry.Data.FindOrAdd(InCacheId.Guid);
+		if (!StructData.IsValid())
+		{
+			if (bValidateWritable)
+			{
+				PCG_EXECUTION_CACHE_VALIDATION_CHECK(InSource);
+			}
+
+			StructData.InitializeAs<T>(InMakeEntry());
+			
+			UE_CLOG(bEnableLogging, LogPCG, Log, TEXT("         [%s] CACHE SET %s"),
+				*InSource->GetExecutionState().GetDebugName(),
+				InCacheId.DebugName ? InCacheId.DebugName : TEXT("unknown"));
+		}
+		return StructData.GetMutable<T>().GetValue();
+	}
+	return InMakeEntry();
+}
diff --git a/Engine/Plugins/PCG/Source/PCG/Public/PCGCommon.h b/Engine/Plugins/PCG/Source/PCG/Public/PCGCommon.h
--- a/Engine/Plugins/PCG/Source/PCG/Public/PCGCommon.h
+++ b/Engine/Plugins/PCG/Source/PCG/Public/PCGCommon.h
@@ -5,9 +5,11 @@
 #include "PCGPointPropertiesTraits.h"
 
 #include "PerQualityLevelProperties.h"
+#include "Async/LockTags.h"
 #include "Async/Mutex.h"
 #include "Async/RecursiveMutex.h"
 #include "Async/SharedMutex.h"
+#include "Async/SharedRecursiveMutex.h"
 #include "HAL/IConsoleManager.h"
 
 #include "PCGCommon.generated.h"
@@ -32,19 +34,13 @@ namespace PCG
 	using FLock = UE::FMutex;
 	using FRecursiveLock = UE::FRecursiveMutex;
 	using FSharedLock = UE::FSharedMutex;
-
-	template<typename LockType>
-	concept IsSharedLock = requires(LockType Lock)
-	{
-		{ Lock.LockShared() };
-		{ Lock.UnlockShared() };
-	};
+	using FRecursiveSharedLock = UE::FSharedRecursiveMutex;
 
 	// Similar to existing UE::TScopeLock but allows optional TryLock in which case caller is responsible to 
 	// check if Lock succeeded or not with IsLocked(). This allows execution code to be skipped if non main thread
 	// tasks fail to get the lock instead of waiting.
 	// Also add an option to have an optional lock, if we have a code path that support both lock and lockless.
-	template<typename LockType> requires (!IsSharedLock<LockType>)
+	template<typename LockType>
 	class TScopeLock final
 	{
 	public:
@@ -95,19 +91,19 @@ namespace PCG
 		bool bShouldLock;
 	};
 
-	template <typename SharedLockType> requires (IsSharedLock<SharedLockType>)
+	template <typename SharedLockType>
 	class TSharedScopeLock final
 	{
 	public:
 		UE_NONCOPYABLE(TSharedScopeLock);
 
 		UE_NODISCARD_CTOR inline explicit TSharedScopeLock(SharedLockType& InSharedLock, bool bInShouldLock = true)
-			: SharedLock(InSharedLock)
+			: SharedLock(InSharedLock, UE::FDeferLock())
 			, bShouldLock(bInShouldLock)
 		{
 			if (bShouldLock)
 			{
-				SharedLock.LockShared();
+				SharedLock.Lock();
 			}
 		}
 
@@ -115,16 +111,16 @@ namespace PCG
 		{
 			if (bShouldLock)
 			{
-				SharedLock.UnlockShared();
+				SharedLock.Unlock();
 			}
 		}
 
 	private:
-		SharedLockType& SharedLock;
+		UE::TDynamicSharedLock<SharedLockType> SharedLock;
 		bool bShouldLock;
 	};
 
-	template <typename SharedLockType> requires (IsSharedLock<SharedLockType>)
+	template <typename SharedLockType>
 	class TUniqueScopeLock final
 	{
 	public:
@@ -909,15 +905,15 @@ private:
 
 #if PCG_EXECUTION_CACHE_VALIDATION_ENABLED
 
-#define PCG_EXECUTION_CACHE_VALIDATION_CREATE_SCOPE(PCGComponent) TGuardValue<bool> ValidationCreateScope(PCGComponent->bCanCreateExecutionCache, true);
-#define PCG_EXECUTION_CACHE_VALIDATION_CREATE_ORIGINAL_SCOPE(PCGComponent) TGuardValue<bool> ValidationCreateOriginalScope(PCGComponent->GetOriginalComponent()->bCanCreateExecutionCache, PCGComponent->bCanCreateExecutionCache);
-#define PCG_EXECUTION_CACHE_VALIDATION_CHECK(PCGComponent) ensureAlways(PCGComponent->bCanCreateExecutionCache || PCGComponent->CurrentGenerationTask == InvalidPCGTaskId);
+#define PCG_EXECUTION_CACHE_VALIDATION_CREATE_SCOPE(Source) TGuardValue<bool> ValidationCreateScope(Source->GetExecutionState().bExecutionCacheWriteEnabled, true);
+#define PCG_EXECUTION_CACHE_VALIDATION_CREATE_ORIGINAL_SCOPE(Source) TGuardValue<bool> ValidationCreateOriginalScope(Source->GetExecutionState().GetOriginalSource()->GetExecutionState().bExecutionCacheWriteEnabled, Source->GetExecutionState().bExecutionCacheWriteEnabled);
+#define PCG_EXECUTION_CACHE_VALIDATION_CHECK(Source) ensureAlways(Source->GetExecutionState().bExecutionCacheWriteEnabled || Source->GetExecutionState().GetGenerationTaskId() == InvalidPCGTaskId);
 
 #else
 
-#define PCG_EXECUTION_CACHE_VALIDATION_CREATE_SCOPE(PCGComponent)
-#define PCG_EXECUTION_CACHE_VALIDATION_CREATE_ORIGINAL_SCOPE(PCGComponent)
-#define PCG_EXECUTION_CACHE_VALIDATION_CHECK(PCGComponent)
+#define PCG_EXECUTION_CACHE_VALIDATION_CREATE_SCOPE(Source)
+#define PCG_EXECUTION_CACHE_VALIDATION_CREATE_ORIGINAL_SCOPE(Source)
+#define PCG_EXECUTION_CACHE_VALIDATION_CHECK(Source)
 
 #endif // PCG_EXECUTION_CACHE_VALIDATION_ENABLED
 
diff --git a/Engine/Plugins/PCG/Source/PCG/Public/PCGComponent.h b/Engine/Plugins/PCG/Source/PCG/Public/PCGComponent.h
--- a/Engine/Plugins/PCG/Source/PCG/Public/PCGComponent.h
+++ b/Engine/Plugins/PCG/Source/PCG/Public/PCGComponent.h
@@ -841,11 +841,6 @@ public:
 	mutable PCGUtils::FExtraCapture ExtraCapture;
 #endif // WITH_EDITOR
 
-public:
-#if PCG_EXECUTION_CACHE_VALIDATION_ENABLED
-	bool bCanCreateExecutionCache = false;
-#endif
-
 private:
 	/** Generalized storage for the execution source (ex. tool data, runtime gen data, manual overrides, etc.) */
 	UPROPERTY()
diff --git a/Engine/Plugins/PCG/Source/PCG/Public/PCGDefaultActorExecutionSource.h b/Engine/Plugins/PCG/Source/PCG/Public/PCGDefaultActorExecutionSource.h
--- a/Engine/Plugins/PCG/Source/PCG/Public/PCGDefaultActorExecutionSource.h
+++ b/Engine/Plugins/PCG/Source/PCG/Public/PCGDefaultActorExecutionSource.h
@@ -48,5 +48,9 @@ protected:
 	UPCGDefaultActorExecutionSource* GetTypedSource() const { return CastChecked<UPCGDefaultActorExecutionSource>(GetSource()); }
 
 private:
+	UPCGData* GetSelfDataInternal() const;
+	FBox GetBoundsInternal() const;
+	FBox GetLocalSpaceBoundsInternal() const;
+
 	AActor* GetActor() const { return GetTypedSource()->GetActor(); }
 };
\ No newline at end of file
diff --git a/Engine/Plugins/PCG/Source/PCG/Public/PCGGraphExecutionStateInterface.h b/Engine/Plugins/PCG/Source/PCG/Public/PCGGraphExecutionStateInterface.h
--- a/Engine/Plugins/PCG/Source/PCG/Public/PCGGraphExecutionStateInterface.h
+++ b/Engine/Plugins/PCG/Source/PCG/Public/PCGGraphExecutionStateInterface.h
@@ -260,6 +260,10 @@ public:
 
 	/** Returns true if the generated artifacts are transient */
 	virtual bool IsInPreviewMode() const { return false; }
+
+#if PCG_EXECUTION_CACHE_VALIDATION_ENABLED
+	bool bExecutionCacheWriteEnabled = false;
+#endif
 };
 
 UINTERFACE(BlueprintType, MinimalAPI)
diff --git a/Engine/Plugins/PCG/Source/PCG/Public/Subsystems/IPCGBaseSubsystem.h b/Engine/Plugins/PCG/Source/PCG/Public/Subsystems/IPCGBaseSubsystem.h
--- a/Engine/Plugins/PCG/Source/PCG/Public/Subsystems/IPCGBaseSubsystem.h
+++ b/Engine/Plugins/PCG/Source/PCG/Public/Subsystems/IPCGBaseSubsystem.h
@@ -6,6 +6,7 @@
 
 #include "PCGCommon.h"
 #include "Elements/PCGActorSelector.h"
+#include "Graph/PCGGraphPerExecutionCache.h"
 #include "Utils/PCGNodeVisualLogs.h"
 #include "UObject/Interface.h"
 
@@ -98,6 +99,53 @@ public:
 	UE_DEPRECATED(5.8, "Use FPCGContext::GetComputeGraph instead, which resolves from the executing graph's compiler cache.")
 	PCG_API UPCGComputeGraph* GetComputeGraph(const UPCGGraph* InGraph, uint32 GridSize, uint32 ComputeGraphIndex);
 
+	/** Returns the cached value if the entry exists for the given source's task, unset otherwise. */
+	template<typename T>
+	static TOptional<typename T::ValueType> GetExecutionCacheEntry(const IPCGGraphExecutionSource* InSource, TPCGPerExecutionCacheId<T> InCacheId)
+	{
+		check(InSource);
+		if (IPCGBaseSubsystem* PCGSubsystem = InSource->GetExecutionState().GetSubsystem())
+		{
+			if (FPCGPerExecutionCache* Cache = PCGSubsystem->GetCacheInternal())
+			{
+				return Cache->GetExecutionCacheEntry(InSource, InCacheId);
+			}
+		}
+
+		return TOptional<typename T::ValueType>{};
+	}
+
+	/** Stores a new value for the given source's task and typed ID. Asserts if the slot is already occupied. */
+	template<typename T>
+	static void SetExecutionCacheEntry(const IPCGGraphExecutionSource* InSource, TPCGPerExecutionCacheId<T> InCacheId, typename T::ValueType InValue, bool bValidateWritable = true)
+	{
+		check(InSource);
+		if (IPCGBaseSubsystem* PCGSubsystem = InSource->GetExecutionState().GetSubsystem())
+		{
+			if (FPCGPerExecutionCache* Cache = PCGSubsystem->GetCacheInternal())
+			{
+				Cache->SetExecutionCacheEntry(InSource, InCacheId, InValue, bValidateWritable);
+			}
+		}
+	}
+
+	/** Returns the cached value if present, otherwise calls InMakeEntry() to produce one.
+	 *  If the source's task is valid the produced value is stored before returning.
+	 *  If no subsystem is reachable from InSource InMakeEntry() is called and its result returned uncached. */
+	template<typename T, typename FuncType>
+	static typename T::ValueType GetOrCreateExecutionCacheValue(const IPCGGraphExecutionSource* InSource, TPCGPerExecutionCacheId<T> InCacheId, const FuncType& InMakeEntry, bool bValidateWritable = true)
+	{
+		check(InSource);
+		if (IPCGBaseSubsystem* PCGSubsystem = InSource->GetExecutionState().GetSubsystem(); PCGSubsystem)
+		{
+			if (FPCGPerExecutionCache* Cache = PCGSubsystem->GetCacheInternal())
+			{
+				return Cache->GetOrCreateExecutionCacheValue(InSource, InCacheId, InMakeEntry, bValidateWritable, PCGSubsystem->IsGraphCacheDebuggingEnabled());
+			}
+		}
+		return InMakeEntry();
+	}
+
 #if WITH_EDITOR
 public:
 	/** Propagate to the graph compiler graph changes */
@@ -163,7 +211,9 @@ protected:
 	PCG_API void CancelGeneration(IPCGGraphExecutionSource* Source, bool bCleanupUnusedResources);
 	
 	virtual void CancelGenerationInternal(IPCGGraphExecutionSource* Source, bool bCleanupUnusedResources) {};
-	
+
+	PCG_API FPCGPerExecutionCache* GetCacheInternal() const;
+
 	TSharedPtr<FPCGGraphExecutor> GraphExecutor;
 
 #if WITH_EDITOR
diff --git a/Engine/Plugins/PCG/Source/PCG/Public/Subsystems/PCGSubsystem.h b/Engine/Plugins/PCG/Source/PCG/Public/Subsystems/PCGSubsystem.h
--- a/Engine/Plugins/PCG/Source/PCG/Public/Subsystems/PCGSubsystem.h
+++ b/Engine/Plugins/PCG/Source/PCG/Public/Subsystems/PCGSubsystem.h
@@ -329,32 +329,6 @@ public:
 	void ResetCacheStats();
 	void LogCacheStats();
 
-	/** Getters for per execution data */
-	PCG_API UPCGData* GetPCGData(FPCGTaskId InGraphExecutionTaskId, bool& bOutFound);
-	PCG_API UPCGData* GetInputPCGData(FPCGTaskId InGraphExecutionTaskId, bool& bOutFound);
-	PCG_API UPCGData* GetSelfPCGData(FPCGTaskId InGraphExecutionTaskId, bool& bOutFound);
-	PCG_API UPCGData* GetLandscapePCGData(FPCGTaskId InGraphExecutionTaskId, bool& bOutFound);
-	PCG_API UPCGData* GetLandscapeHeightPCGData(FPCGTaskId InGraphExecutionTaskId, bool& bOutFound);
-	PCG_API UPCGData* GetOriginalSelfPCGData(FPCGTaskId InGraphExecutionTaskId, bool& bOutFound);
-	PCG_API TOptional<FBox> GetBounds(FPCGTaskId InGraphExecutionTaskId);
-	PCG_API TOptional<FBox> GetOriginalBounds(FPCGTaskId InGraphExecutionTaskId);
-	PCG_API TOptional<FBox> GetLocalSpaceBounds(FPCGTaskId InGraphExecutionTaskId);
-	PCG_API TOptional<FBox> GetOriginalLocalSpaceBounds(FPCGTaskId InGraphExecutionTaskId);
-	PCG_API TOptional<FBox> GetTotalBounds(FPCGTaskId InGraphExecutionTaskId);
-
-	/** Setters for per execution data. Usually called from execution state ExecutePreGraph implementations. */
-	PCG_API void SetPCGData(FPCGTaskId InGraphExecutionTaskId, UPCGData* InData);
-	PCG_API void SetInputPCGData(FPCGTaskId InGraphExecutionTaskId, UPCGData* InData);
-	PCG_API void SetSelfPCGData(FPCGTaskId InGraphExecutionTaskId, UPCGData* InData);
-	PCG_API void SetLandscapePCGData(FPCGTaskId InGraphExecutionTaskId, UPCGData* InData);
-	PCG_API void SetLandscapeHeightPCGData(FPCGTaskId InGraphExecutionTaskId, UPCGData* InData);
-	PCG_API void SetOriginalSelfPCGData(FPCGTaskId InGraphExecutionTaskId, UPCGData* InData);
-	PCG_API void SetBounds(FPCGTaskId InGraphExecutionTaskId, const FBox& InBounds);
-	PCG_API void SetOriginalBounds(FPCGTaskId InGraphExecutionTaskId, const FBox& InBounds);
-	PCG_API void SetLocalSpaceBounds(FPCGTaskId InGraphExecutionTaskId, const FBox& InBounds);
-	PCG_API void SetOriginalLocalSpaceBounds(FPCGTaskId InGraphExecutionTaskId, const FBox& InBounds);
-	PCG_API void SetTotalBounds(FPCGTaskId InGraphExecutionTaskId, const FBox& InBounds);
-
 #if WITH_EDITOR
 public:
 	/** Schedule refresh on the current or next frame */
diff --git a/Engine/Plugins/PCG/Tests/Private/Cache/PCGGraphExecutionCacheTestTypes.h b/Engine/Plugins/PCG/Tests/Private/Cache/PCGGraphExecutionCacheTestTypes.h
new file mode 100644
--- /dev/null
+++ b/Engine/Plugins/PCG/Tests/Private/Cache/PCGGraphExecutionCacheTestTypes.h
@@ -0,0 +1,53 @@
+// Copyright Epic Games, Inc. All Rights Reserved.
+
+#pragma once
+
+#include "Graph/PCGGraphPerExecutionCache.h"
+#include "PCGDefaultExecutionSource.h"
+
+#include "PCGGraphExecutionCacheTestTypes.generated.h"
+
+/** Minimal cache entry used exclusively in PCGGraphExecutionCacheTests. */
+USTRUCT()
+struct FPCGTestCacheEntry : public FPCGPerExecutionCacheData
+{
+	GENERATED_BODY()
+
+	using ValueType = int32;
+
+	FPCGTestCacheEntry() = default;
+	FPCGTestCacheEntry(ValueType InValue) : Value(InValue) {}
+
+	ValueType GetValue() const { return Value; }
+
+	ValueType Value = 0;
+};
+
+/** Execution state stub for cache tests — overrides only GetGenerationTaskId(). */
+class FPCGTestExecutionState final : public FPCGDefaultExecutionState
+{
+public:
+	explicit FPCGTestExecutionState(UPCGDefaultExecutionSource* InSource) : FPCGDefaultExecutionState(InSource) {}
+
+	virtual FPCGTaskId GetGenerationTaskId() const override { return TaskId; }
+
+	FPCGTaskId TaskId = InvalidPCGTaskId;
+};
+
+/** Execution source stub for cache tests — exposes a settable task ID. */
+UCLASS()
+class UPCGTestExecutionSource : public UPCGDefaultExecutionSource
+{
+	GENERATED_BODY()
+
+public:
+	UPCGTestExecutionSource()
+	{
+		State = MakeUnique<FPCGTestExecutionState>(this);
+	}
+
+	void SetTaskId(FPCGTaskId InTaskId)
+	{
+		static_cast<FPCGTestExecutionState&>(*State).TaskId = InTaskId;
+	}
+};
diff --git a/Engine/Plugins/PCG/Tests/Private/Cache/PCGGraphExecutionCacheTests.cpp b/Engine/Plugins/PCG/Tests/Private/Cache/PCGGraphExecutionCacheTests.cpp
new file mode 100644
--- /dev/null
+++ b/Engine/Plugins/PCG/Tests/Private/Cache/PCGGraphExecutionCacheTests.cpp
@@ -0,0 +1,355 @@
+// Copyright Epic Games, Inc. All Rights Reserved.
+
+// Test summary
+// ┌─────────────────────────────────────────────────────┬───────────────────────────────────────────────────────────────-─┐
+// │ Test name                                           │ What it covers                                                  │
+// ├─────────────────────────────────────────────────────┼────────────────────────────────────────────────────────────────-┤
+// │ ExecutionCache::GetExecutionCacheEntry              │ Returns unset optional when no entry exists or task is invalid; │
+// │                                                     │ returns the stored value after SetExecutionCacheEntry.          │
+// |                                                     |                                                                 |
+// │ ExecutionCache::TaskIsolation                       │ Entries stored under task A are not visible under task B.       │
+// |                                                     |                                                                 |
+// │ ExecutionCache::CacheIdIsolation                    │ Entries stored under one typed ID are not visible under a       │
+// │                                                     │ different ID for the same task.                                 │
+// |                                                     |                                                                 |
+// │ ExecutionCache::Clear                               │ After Clear(), the entries map is empty and no previously       │
+// │                                                     │ stored value is retrievable.                                    │
+// |                                                     |                                                                 |
+// │ ExecutionCache::GetOrCreateWithNoTask               │ With an invalid task ID the lambda is called directly and its   │
+// │                                                     │ result returned without being stored in the cache.              │
+// |                                                     |                                                                 |
+// │ ExecutionCache::CachingWithValidTask                │ First call (cache miss) invokes the lambda and stores the       │
+// │                                                     │ result; subsequent calls return the cached value without        │
+// │                                                     │ invoking the lambda again.                                      │
+// |                                                     |                                                                 |
+// │ ExecutionCache::ComponentGettersWithNoWorld         │ UPCGComponent data and bounds getters do not crash when there   │
+// │                                                     │ is no world (and therefore no subsystem), and member-level      │
+// │                                                     │ caches make repeated calls return the same pointer.             │
+// |                                                     |                                                                 |
+// │ ExecutionCache::GrowthDoesNotInvalidateEntries      │ Adding many entries after an initial write — growing both the   │
+// │                                                     │ outer task map and the inner per-task GUID map — leaves the     │
+// │                                                     │ original entry intact and readable.                             │
+// └─────────────────────────────────────────────────────┴───────────────────────────────────────────────────────────────-─┘
+
+#include "PCGTestsCommon.h"
+
+#include "PCGComponent.h"
+#include "Cache/PCGGraphExecutionCacheTestTypes.h"
+#include "Graph/PCGGraphPerExecutionCache.h"
+
+#include "UObject/UObjectGlobals.h"
+
+#include <catch2/catch_test_macros.hpp>
+
+namespace PCGExecutionCacheTests
+{
+	/** Cache IDs that are not part of the production PCGPerExecutionCacheGuids namespace. */
+	constexpr TPCGPerExecutionCacheId<FPCGTestCacheEntry> TestEntry { FGuid(0xC8A1F3E2, 0x4D7B9C05, 0xA362E817, 0x5F094D3B) };
+	constexpr TPCGPerExecutionCacheId<FPCGTestCacheEntry> TestEntry2{ FGuid(0x1A2B3C4D, 0x5E6F7A8B, 0x9C0D1E2F, 0x3A4B5C6D) };
+}
+
+// ---------------------------------------------------------------------------
+// FPCGPerExecutionCache — GetExecutionCacheEntry
+// ---------------------------------------------------------------------------
+// Verifies that GetExecutionCacheEntry returns an unset optional when no entry
+// has been stored, and the correct value after SetExecutionCacheEntry.
+
+TEST_CASE_METHOD(PCGTests::FPCGBaseTest, "PCG::ExecutionCache::GetExecutionCacheEntry", "[PCG][ExecutionCache]")
+{
+	FPCGPerExecutionCache Cache;
+
+	UPCGTestExecutionSource* Source = NewObject<UPCGTestExecutionSource>();
+	Source->SetTaskId(1);
+
+	UPCGTestExecutionSource* NoTaskSource = NewObject<UPCGTestExecutionSource>();
+	// NoTaskSource has TaskId == InvalidPCGTaskId by default.
+
+	SECTION("Returns unset optional when no entry exists for the task")
+	{
+		const TOptional<int32> Result = Cache.GetExecutionCacheEntry(Source, PCGExecutionCacheTests::TestEntry);
+		CHECK_FALSE(Result.IsSet());
+	}
+
+	SECTION("Returns unset optional when the source has no active task")
+	{
+		const TOptional<int32> Result = Cache.GetExecutionCacheEntry(NoTaskSource, PCGExecutionCacheTests::TestEntry);
+		CHECK_FALSE(Result.IsSet());
+	}
+
+	SECTION("Returns the stored value after SetExecutionCacheEntry")
+	{
+		Cache.SetExecutionCacheEntry(Source, PCGExecutionCacheTests::TestEntry, 42, /*bValidateWritable=*/false);
+
+		const TOptional<int32> Result = Cache.GetExecutionCacheEntry(Source, PCGExecutionCacheTests::TestEntry);
+		REQUIRE(Result.IsSet());
+		CHECK(Result.GetValue() == 42);
+	}
+}
+
+// ---------------------------------------------------------------------------
+// FPCGPerExecutionCache — task isolation
+// ---------------------------------------------------------------------------
+// Entries stored under one task ID must not be visible under a different task ID.
+
+TEST_CASE_METHOD(PCGTests::FPCGBaseTest, "PCG::ExecutionCache::TaskIsolation", "[PCG][ExecutionCache]")
+{
+	FPCGPerExecutionCache Cache;
+
+	UPCGTestExecutionSource* SourceA = NewObject<UPCGTestExecutionSource>();
+	SourceA->SetTaskId(1);
+
+	UPCGTestExecutionSource* SourceB = NewObject<UPCGTestExecutionSource>();
+	SourceB->SetTaskId(2);
+
+	Cache.SetExecutionCacheEntry(SourceA, PCGExecutionCacheTests::TestEntry, 10, /*bValidateWritable=*/false);
+
+	const TOptional<int32> ResultA = Cache.GetExecutionCacheEntry(SourceA, PCGExecutionCacheTests::TestEntry);
+	REQUIRE(ResultA.IsSet());
+	CHECK(ResultA.GetValue() == 10);
+
+	const TOptional<int32> ResultB = Cache.GetExecutionCacheEntry(SourceB, PCGExecutionCacheTests::TestEntry);
+	CHECK_FALSE(ResultB.IsSet());
+}
+
+// ---------------------------------------------------------------------------
+// FPCGPerExecutionCache — cache-ID isolation
+// ---------------------------------------------------------------------------
+// Entries stored under one typed ID must not be visible under a different ID
+// for the same task.
+
+TEST_CASE_METHOD(PCGTests::FPCGBaseTest, "PCG::ExecutionCache::CacheIdIsolation", "[PCG][ExecutionCache]")
+{
+	FPCGPerExecutionCache Cache;
+
+	UPCGTestExecutionSource* Source = NewObject<UPCGTestExecutionSource>();
+	Source->SetTaskId(1);
+
+	Cache.SetExecutionCacheEntry(Source, PCGExecutionCacheTests::TestEntry, 7, /*bValidateWritable=*/false);
+
+	const TOptional<int32> Result1 = Cache.GetExecutionCacheEntry(Source, PCGExecutionCacheTests::TestEntry);
+	REQUIRE(Result1.IsSet());
+	CHECK(Result1.GetValue() == 7);
+
+	const TOptional<int32> Result2 = Cache.GetExecutionCacheEntry(Source, PCGExecutionCacheTests::TestEntry2);
+	CHECK_FALSE(Result2.IsSet());
+}
+
+// ---------------------------------------------------------------------------
+// FPCGPerExecutionCache — Clear
+// ---------------------------------------------------------------------------
+// After Clear(), no previously stored entries should be retrievable.
+
+TEST_CASE_METHOD(PCGTests::FPCGBaseTest, "PCG::ExecutionCache::Clear", "[PCG][ExecutionCache]")
+{
+	FPCGPerExecutionCache Cache;
+
+	UPCGTestExecutionSource* Source = NewObject<UPCGTestExecutionSource>();
+	Source->SetTaskId(1);
+
+	Cache.SetExecutionCacheEntry(Source, PCGExecutionCacheTests::TestEntry, 5, /*bValidateWritable=*/false);
+	REQUIRE(Cache.GetExecutionCacheEntry(Source, PCGExecutionCacheTests::TestEntry).IsSet());
+
+	Cache.Clear();
+
+	CHECK(Cache.Entries.IsEmpty());
+	CHECK_FALSE(Cache.GetExecutionCacheEntry(Source, PCGExecutionCacheTests::TestEntry).IsSet());
+}
+
+// ---------------------------------------------------------------------------
+// FPCGPerExecutionCache — GetOrCreateExecutionCacheValue with no active task
+// ---------------------------------------------------------------------------
+// Verifies that with no active task the lambda is called directly and the
+// result is returned without being stored in the cache.
+
+TEST_CASE_METHOD(PCGTests::FPCGBaseTest, "PCG::ExecutionCache::GetOrCreateWithNoTask", "[PCG][ExecutionCache]")
+{
+	FPCGPerExecutionCache Cache;
+
+	UPCGTestExecutionSource* NoTaskSource = NewObject<UPCGTestExecutionSource>();
+	// NoTaskSource has TaskId == InvalidPCGTaskId by default.
+
+	SECTION("Lambda is called and its value is returned")
+	{
+		bool bLambdaCalled = false;
+		const FBox TestBox(FVector(-100.f), FVector(100.f));
+
+		const FBox Result = Cache.GetOrCreateExecutionCacheValue(NoTaskSource, PCGPerExecutionCacheGuids::Bounds,
+			[&bLambdaCalled, TestBox]()
+			{
+				bLambdaCalled = true;
+				return TestBox;
+			}, /*bValidateWritable=*/false);
+
+		REQUIRE(bLambdaCalled);
+		REQUIRE(Result == TestBox);
+	}
+
+	SECTION("Nothing is stored in the cache when there is no active task")
+	{
+		Cache.GetOrCreateExecutionCacheValue(NoTaskSource, PCGPerExecutionCacheGuids::Bounds,
+			[]() { return FBox(FVector(-1.f), FVector(1.f)); }, /*bValidateWritable=*/false);
+
+		CHECK(Cache.Entries.IsEmpty());
+	}
+}
+
+// ---------------------------------------------------------------------------
+// FPCGPerExecutionCache — caching behaviour with a valid task ID
+// ---------------------------------------------------------------------------
+// Uses a custom cache entry type (FPCGTestCacheEntry) to verify that:
+//   - the first call invokes the lambda and stores the result, and
+//   - subsequent calls return the cached value without invoking the lambda.
+
+TEST_CASE_METHOD(PCGTests::FPCGBaseTest, "PCG::ExecutionCache::CachingWithValidTask", "[PCG][ExecutionCache]")
+{
+	FPCGPerExecutionCache Cache;
+	constexpr FPCGTaskId TaskId = 1;
+
+	UPCGTestExecutionSource* Source = NewObject<UPCGTestExecutionSource>();
+	Source->SetTaskId(TaskId);
+
+	int32 LambdaCallCount = 0;
+
+	// First call: cache miss — lambda must be invoked and the value stored.
+	const int32 First = Cache.GetOrCreateExecutionCacheValue(Source, PCGExecutionCacheTests::TestEntry,
+		[&LambdaCallCount]() { ++LambdaCallCount; return 99; }, /*bValidateWritable=*/false);
+
+	REQUIRE(LambdaCallCount == 1);
+	REQUIRE(First == 99);
+
+	SECTION("Lambda is not called again on a cache hit")
+	{
+		const int32 Second = Cache.GetOrCreateExecutionCacheValue(Source, PCGExecutionCacheTests::TestEntry,
+			[&LambdaCallCount]() { ++LambdaCallCount; return 0; }, /*bValidateWritable=*/false);
+
+		CHECK(LambdaCallCount == 1);
+		CHECK(Second == 99);
+	}
+
+	SECTION("Entry is stored under the correct task ID")
+	{
+		REQUIRE(Cache.Entries.Contains(TaskId));
+		CHECK(Cache.Entries[TaskId].Data.Contains(PCGExecutionCacheTests::TestEntry.Guid));
+	}
+}
+
+// ---------------------------------------------------------------------------
+// UPCGComponent getters — no world (component not placed in a level)
+// ---------------------------------------------------------------------------
+// Regression test: before the GetOrCreateExecutionCacheValue refactor some
+// getters had an early "if (CurrentGenerationTask != InvalidPCGTaskId)"
+// guard that returned nullptr without invoking the lambda.
+//
+// A bare AActor with no level/world is used as the owner so that
+// GetOwner() is non-null (required by the data getter lambdas) while
+// GetWorld() returns null (no subsystem reachable).  With no subsystem the
+// static GetOrCreateExecutionCacheValue overload calls the lambda directly
+// on every invocation.
+
+TEST_CASE_METHOD(PCGTests::FPCGBaseTest, "PCG::ExecutionCache::ComponentGettersWithNoWorld", "[PCG][ExecutionCache]")
+{
+	AActor* Actor = NewObject<AActor>();
+	UPCGComponent* Component = NewObject<UPCGComponent>(Actor);
+
+	REQUIRE(Component->GetWorld() == nullptr);
+	REQUIRE_FALSE(Component->IsGenerating());
+
+	SECTION("Bounds getters do not crash")
+	{
+		// The actor has no components so the helpers return uninitialized
+		// boxes, but must not crash.
+		Component->GetGridBounds();
+		Component->GetOriginalGridBounds();
+		Component->GetLocalSpaceBounds();
+		Component->GetOriginalLocalSpaceBounds();
+		Component->GetTotalBounds();
+	}
+
+	SECTION("PCG data getters do not crash")
+	{
+		// All data getters must invoke the lambda path rather than
+		// short-circuiting to null.  A bare actor produces no spatial data
+		// so null returns are expected, but the calls must not crash.
+		Component->GetPCGData();
+		Component->GetInputPCGData();
+		Component->GetActorPCGData();
+		Component->GetOriginalActorPCGData();
+		Component->GetLandscapePCGData();
+		Component->GetLandscapeHeightPCGData();
+	}
+
+	SECTION("Data getter results are consistent across repeated calls")
+	{
+		// Component-level member caches (e.g. CachedActorData) must make
+		// the same pointer be returned on repeated calls even without an
+		// execution-cache task ID.
+		UPCGData* const ActorData0 = Component->GetActorPCGData();
+		UPCGData* const ActorData1 = Component->GetActorPCGData();
+		CHECK(ActorData0 == ActorData1);
+
+		UPCGData* const PCGData0 = Component->GetPCGData();
+		UPCGData* const PCGData1 = Component->GetPCGData();
+		CHECK(PCGData0 == PCGData1);
+	}
+}
+
+// ---------------------------------------------------------------------------
+// FPCGPerExecutionCache — cache growth does not invalidate existing entries
+// ---------------------------------------------------------------------------
+// TMap may rehash its internal storage when new elements are inserted.
+// Verifies that entries written before a growth event are still present and
+// hold their original values afterwards.
+//
+// Two growth axes are tested:
+//   - Outer map growth: many distinct task IDs are added after the sentinel,
+//     forcing the FPCGTaskId → FPCGPerExecutionCacheEntry map to rehash.
+//   - Inner map growth: many distinct GUIDs are added under the sentinel's
+//     task, forcing the per-task FGuid → TInstancedStruct map to rehash.
+
+TEST_CASE_METHOD(PCGTests::FPCGBaseTest, "PCG::ExecutionCache::GrowthDoesNotInvalidateEntries", "[PCG][ExecutionCache]")
+{
+	constexpr int32    SentinelValue  = 12345;
+	constexpr FPCGTaskId SentinelTaskId = 1;
+	constexpr int32    NumExtraEntries = 128;
+
+	SECTION("Outer map growth (many task IDs)")
+	{
+		FPCGPerExecutionCache Cache;
+
+		UPCGTestExecutionSource* SentinelSource = NewObject<UPCGTestExecutionSource>();
+		SentinelSource->SetTaskId(SentinelTaskId);
+		Cache.SetExecutionCacheEntry(SentinelSource, PCGExecutionCacheTests::TestEntry, SentinelValue, /*bValidateWritable=*/false);
+
+		// Insert entries under distinct task IDs to grow the outer TMap.
+		for (int32 i = 0; i < NumExtraEntries; ++i)
+		{
+			UPCGTestExecutionSource* Source = NewObject<UPCGTestExecutionSource>();
+			Source->SetTaskId(static_cast<FPCGTaskId>(SentinelTaskId + 1 + i));
+			Cache.SetExecutionCacheEntry(Source, PCGExecutionCacheTests::TestEntry, i, /*bValidateWritable=*/false);
+		}
+
+		const TOptional<int32> Result = Cache.GetExecutionCacheEntry(SentinelSource, PCGExecutionCacheTests::TestEntry);
+		REQUIRE(Result.IsSet());
+		CHECK(Result.GetValue() == SentinelValue);
+	}
+
+	SECTION("Inner map growth (many GUIDs under one task)")
+	{
+		FPCGPerExecutionCache Cache;
+
+		UPCGTestExecutionSource* Source = NewObject<UPCGTestExecutionSource>();
+		Source->SetTaskId(SentinelTaskId);
+		Cache.SetExecutionCacheEntry(Source, PCGExecutionCacheTests::TestEntry, SentinelValue, /*bValidateWritable=*/false);
+
+		// Insert entries under distinct GUIDs in the same task to grow the inner TMap.
+		for (int32 i = 0; i < NumExtraEntries; ++i)
+		{
+			const TPCGPerExecutionCacheId<FPCGTestCacheEntry> DynamicId{ FGuid(0xDEADBEEF, 0, 0, static_cast<uint32>(i)) };
+			Cache.SetExecutionCacheEntry(Source, DynamicId, i, /*bValidateWritable=*/false);
+		}
+
+		const TOptional<int32> Result = Cache.GetExecutionCacheEntry(Source, PCGExecutionCacheTests::TestEntry);
+		REQUIRE(Result.IsSet());
+		CHECK(Result.GetValue() == SentinelValue);
+	}
+}
