[PCG] Fix reentrance issue (regression from 51758684)
- find data after call to InMakeEntry() as this can be a reentrant call invalidating existing references
- add new test to cover this code path

#Synced-CL 51822038

diff --git a/Engine/Plugins/PCG/Source/PCG/Public/Graph/PCGGraphPerExecutionCache.h b/Engine/Plugins/PCG/Source/PCG/Public/Graph/PCGGraphPerExecutionCache.h
--- a/Engine/Plugins/PCG/Source/PCG/Public/Graph/PCGGraphPerExecutionCache.h
+++ b/Engine/Plugins/PCG/Source/PCG/Public/Graph/PCGGraphPerExecutionCache.h
@@ -189,22 +189,29 @@ typename T::ValueType FPCGPerExecutionCache::GetOrCreateExecutionCacheValue(cons
 		}
 
 		PCG::TUniqueScopeLock WriteLock(Lock);
-		FPCGPerExecutionCacheEntry& Entry = Entries.FindOrAdd(TaskId);
-		TInstancedStruct<FPCGPerExecutionCacheData>& StructData = Entry.Data.FindOrAdd(InCacheId.Guid);
-		if (!StructData.IsValid())
+		FPCGPerExecutionCacheEntry* Entry = &Entries.FindOrAdd(TaskId);
+		TInstancedStruct<FPCGPerExecutionCacheData>* StructData = &Entry->Data.FindOrAdd(InCacheId.Guid);
+		if (!StructData->IsValid())
 		{
 			if (bValidateWritable)
 			{
 				PCG_EXECUTION_CACHE_VALIDATION_CHECK(InSource);
 			}
 
-			StructData.InitializeAs<T>(InMakeEntry());
+			typename T::ValueType EntryValue = InMakeEntry();
+
+			// InMakeEntry() call can cause a reentrant GetOrCreateExecutionCacheValue() call which is supported because FPCGPerExecutionCache::Lock is a PCG::FRecursiveSharedLock 
+			// but it can invalidate our cached pointers (Entry & StructData) so find them again before intializing the struct data.
+			Entry = &Entries.FindChecked(TaskId);
+			StructData = &Entry->Data.FindChecked(InCacheId.Guid);
+
+			StructData->InitializeAs<T>(MoveTemp(EntryValue));
 			
 			UE_CLOG(bEnableLogging, LogPCG, Log, TEXT("         [%s] CACHE SET %s"),
 				*InSource->GetExecutionState().GetDebugName(),
 				InCacheId.DebugName ? InCacheId.DebugName : TEXT("unknown"));
 		}
-		return StructData.GetMutable<T>().GetValue();
+		return StructData->GetMutable<T>().GetValue();
 	}
 	return InMakeEntry();
 }
diff --git a/Engine/Plugins/PCG/Tests/Private/Cache/PCGGraphExecutionCacheTests.cpp b/Engine/Plugins/PCG/Tests/Private/Cache/PCGGraphExecutionCacheTests.cpp
--- a/Engine/Plugins/PCG/Tests/Private/Cache/PCGGraphExecutionCacheTests.cpp
+++ b/Engine/Plugins/PCG/Tests/Private/Cache/PCGGraphExecutionCacheTests.cpp
@@ -29,6 +29,11 @@
 // │ ExecutionCache::GrowthDoesNotInvalidateEntries      │ Adding many entries after an initial write — growing both the   │
 // │                                                     │ outer task map and the inner per-task GUID map — leaves the     │
 // │                                                     │ original entry intact and readable.                             │
+// |                                                     |                                                                 |
+// │ ExecutionCache::                                    │ When InMakeEntry() itself calls GetOrCreateExecutionCacheValue  │
+// │   ReentrantMakeEntryDoesNotCorruptPointers          │ and causes TMap rehashing, the outer call's Entry/StructData    │
+// │                                                     │ pointers are refreshed so the computed value is stored          │
+// │                                                     │ correctly in the cache.                                         │
 // └─────────────────────────────────────────────────────┴───────────────────────────────────────────────────────────────-─┘
 
 #include "PCGTestsCommon.h"
@@ -353,3 +358,115 @@ TEST_CASE_METHOD(PCGTests::FPCGBaseTest, "PCG::ExecutionCache::GrowthDoesNotInva
 		CHECK(Result.GetValue() == SentinelValue);
 	}
 }
+
+// ---------------------------------------------------------------------------
+// FPCGPerExecutionCache — reentrant InMakeEntry does not corrupt cached pointers
+// ---------------------------------------------------------------------------
+// Regression test: when InMakeEntry() itself calls GetOrCreateExecutionCacheValue(),
+// TMap rehashing can occur, invalidating the raw Entry/StructData pointers captured
+// before the call.  The fix re-finds those pointers after InMakeEntry() returns so
+// the computed value is written to the correct location in the cache.
+//
+// Two axes are tested:
+//   - Inner TMap growth: the lambda inserts many GUIDs under the same task ID,
+//     growing Entry->Data and potentially invalidating StructData.
+//   - Outer TMap growth: the lambda inserts entries under many different task IDs,
+//     growing Entries and potentially invalidating Entry.
+//
+// In both cases the outer entry must be stored correctly after the call.
+
+TEST_CASE_METHOD(PCGTests::FPCGBaseTest, "PCG::ExecutionCache::ReentrantMakeEntryDoesNotCorruptPointers", "[PCG][ExecutionCache]")
+{
+	constexpr FPCGTaskId TaskId         = 1;
+	constexpr int32      OuterValue     = 42;
+	constexpr int32      InnerValue     = 99;
+	// 128 reentrant insertions is more than enough to trigger TMap rehashing and
+	// invalidate any raw pointers captured before the growth event.
+	constexpr int32      NumInnerEntries = 128;
+
+	SECTION("Inner TMap growth (same task, many GUIDs)")
+	{
+		FPCGPerExecutionCache Cache;
+
+		UPCGTestExecutionSource* Source = NewObject<UPCGTestExecutionSource>();
+		Source->SetTaskId(TaskId);
+
+		// Build the dynamic IDs ahead of time so we can re-query them after the outer call.
+		TArray<TPCGPerExecutionCacheId<FPCGTestCacheEntry>> DynamicIds;
+		DynamicIds.Reserve(NumInnerEntries);
+		for (int32 i = 0; i < NumInnerEntries; ++i)
+		{
+			DynamicIds.Add(TPCGPerExecutionCacheId<FPCGTestCacheEntry>{ FGuid(0xDEADBEEF, 0, 0, static_cast<uint32>(i)) });
+		}
+
+		// The lambda inserts many extra GUIDs under the same task, growing Entry->Data
+		// and potentially invalidating the StructData pointer held by the outer call.
+		const int32 Result = Cache.GetOrCreateExecutionCacheValue(Source, PCGExecutionCacheTests::TestEntry,
+			[&Cache, Source, &DynamicIds]()
+			{
+				for (int32 i = 0; i < NumInnerEntries; ++i)
+				{
+					Cache.GetOrCreateExecutionCacheValue(Source, DynamicIds[i],
+						[i]() { return InnerValue + i; }, /*bValidateWritable=*/false);
+				}
+				return OuterValue;
+			}, /*bValidateWritable=*/false);
+
+		CHECK(Result == OuterValue);
+
+		const TOptional<int32> Cached = Cache.GetExecutionCacheEntry(Source, PCGExecutionCacheTests::TestEntry);
+		REQUIRE(Cached.IsSet());
+		CHECK(Cached.GetValue() == OuterValue);
+
+		for (int32 i = 0; i < NumInnerEntries; ++i)
+		{
+			const TOptional<int32> InnerCached = Cache.GetExecutionCacheEntry(Source, DynamicIds[i]);
+			REQUIRE(InnerCached.IsSet());
+			CHECK(InnerCached.GetValue() == InnerValue + i);
+		}
+	}
+
+	SECTION("Outer TMap growth (lambda inserts entries under many task IDs)")
+	{
+		FPCGPerExecutionCache Cache;
+
+		UPCGTestExecutionSource* Source = NewObject<UPCGTestExecutionSource>();
+		Source->SetTaskId(TaskId);
+
+		// Build the inner sources ahead of time so we can re-query them after the outer call.
+		TArray<UPCGTestExecutionSource*> InnerSources;
+		InnerSources.Reserve(NumInnerEntries);
+		for (int32 i = 0; i < NumInnerEntries; ++i)
+		{
+			UPCGTestExecutionSource* InnerSource = NewObject<UPCGTestExecutionSource>();
+			InnerSource->SetTaskId(static_cast<FPCGTaskId>(TaskId + 1 + i));
+			InnerSources.Add(InnerSource);
+		}
+
+		// The lambda inserts entries under many distinct task IDs, growing Entries and
+		// potentially invalidating the Entry pointer held by the outer call.
+		const int32 Result = Cache.GetOrCreateExecutionCacheValue(Source, PCGExecutionCacheTests::TestEntry,
+			[&Cache, &InnerSources]()
+			{
+				for (int32 i = 0; i < NumInnerEntries; ++i)
+				{
+					Cache.GetOrCreateExecutionCacheValue(InnerSources[i], PCGExecutionCacheTests::TestEntry2,
+						[i]() { return InnerValue + i; }, /*bValidateWritable=*/false);
+				}
+				return OuterValue;
+			}, /*bValidateWritable=*/false);
+
+		CHECK(Result == OuterValue);
+
+		const TOptional<int32> Cached = Cache.GetExecutionCacheEntry(Source, PCGExecutionCacheTests::TestEntry);
+		REQUIRE(Cached.IsSet());
+		CHECK(Cached.GetValue() == OuterValue);
+
+		for (int32 i = 0; i < NumInnerEntries; ++i)
+		{
+			const TOptional<int32> InnerCached = Cache.GetExecutionCacheEntry(InnerSources[i], PCGExecutionCacheTests::TestEntry2);
+			REQUIRE(InnerCached.IsSet());
+			CHECK(InnerCached.GetValue() == InnerValue + i);
+		}
+	}
+}
