[PCG] Implement type erased accessors

* Preparation work to support any type getter/setter with accessors
* Fix attribute descriptor for legacy types, so we force FVector for example to be Vector and not "Struct + TBaseStructure<FVector>::Get()"
* Update the traits to build a constructible/broadcastable function matrix of our types
* Same for fundamental types, to construct/destruct and have the element size
* Moved the InValues/OutValues concepts from Generic attributes to metadata common for shareability with the accessors
* Move some concepts in its own helper.
* Mimic AddValues/GetValues from Generic attribute for SetRange/GetRange, for the different usecases.
* Remove template specialization and virtual function for all supported types for GetRange/SetRange
* Remove InternalGetRange/InternalSetRange in the templated version, now done with the typed erased version GetRange_Internal/SetRange_Internal
* Forward the virtual call GetRangeVirtual/SetRangeVirtual to the typed call, which should yield the exact same behavior.
* Convert existing unit tests for accessor into low level tests, and extend it with other usecases.

#rnx
#Synced-CL 52060281

diff --git a/Engine/Plugins/PCG/Source/PCG/Private/Metadata/Accessors/IPCGAttributeAccessor.cpp b/Engine/Plugins/PCG/Source/PCG/Private/Metadata/Accessors/IPCGAttributeAccessor.cpp
--- a/Engine/Plugins/PCG/Source/PCG/Private/Metadata/Accessors/IPCGAttributeAccessor.cpp
+++ b/Engine/Plugins/PCG/Source/PCG/Private/Metadata/Accessors/IPCGAttributeAccessor.cpp
@@ -2,11 +2,124 @@
 
 #include "Metadata/Accessors/IPCGAttributeAccessor.h"
 
-// Specialization of IPCGAttributeAccessor::Get and IPCGAttributeAccessor::Set, for all supported types.
-
-#include UE_INLINE_GENERATED_CPP_BY_NAME(IPCGAttributeAccessor)
-#define IACCESSOR_DECL(T) \
-template <> bool IPCGAttributeAccessor::GetRange<T>(TArrayView<T> OutValues, int32 Index, const IPCGAttributeAccessorKeys& Keys, EPCGAttributeAccessorFlags Flags) const { return GetRange##T(OutValues, Index, Keys, Flags); } \
-template <> bool IPCGAttributeAccessor::SetRange<T>(TArrayView<const T> InValues, int32 Index, IPCGAttributeAccessorKeys& Keys, EPCGAttributeAccessorFlags Flags) { return SetRange##T(InValues, Index, Keys, Flags); }
-PCG_FOREACH_SUPPORTEDTYPES(IACCESSOR_DECL);
-#undef IACCESSOR_DECL
+namespace PCGAttributeAccessorPrivate
+{
+	constexpr int32 MaxAlignment = alignof(FTransform);
+	constexpr int32 AlignedBlockSize = MaxAlignment;
+	constexpr int32 MaxInlineMem = 1024 / AlignedBlockSize; // 1KB
+	
+	// Making sure we can accommodate our big types that have different alignment like FTransform or FVector4
+	using FRawBuffer = TArray<TAlignedBytes<AlignedBlockSize, MaxAlignment>, TInlineAllocator<MaxInlineMem>>;
+	
+	int32 GetAllocationSize(const int32 ElementSize, const int32 Count)
+	{
+		return FMath::DivideAndRoundUp(ElementSize * Count, AlignedBlockSize);
+	}
+}
+
+bool IPCGAttributeAccessor::GetRange_Internal(PCG::Private::FOutValues OutValues, int16 ExpectedType, int32 Count, int32 Index, const IPCGAttributeAccessorKeys& Keys, EPCGAttributeAccessorFlags Flags, PCG::Private::EPCGAttributeAccessorOperation OpType) const
+{
+	using namespace PCG::Private;
+	
+	switch (OpType)
+	{
+	case EPCGAttributeAccessorOperation::SameType:
+	{
+		return GetRangeVirtual(MoveTemp(OutValues), Count, Index, Keys, Flags);
+	}
+	case EPCGAttributeAccessorOperation::Construct: // Fall-through
+	case EPCGAttributeAccessorOperation::Broadcast:
+	{
+		check(OutValues.IsType<FOutValuesByValue>() && IsPCGType(ExpectedType));
+			
+		PCGAttributeAccessorPrivate::FRawBuffer Buffer{}; // 1KB
+		const uint16 UnderlyingValueType = static_cast<uint16>(UnderlyingDesc.ValueType);
+			
+		const int32 ElementSize = GetElementSize(UnderlyingValueType);
+		if (!ensure(ElementSize > 0))
+		{
+			return false;
+		}
+			
+		Buffer.SetNumUninitialized(PCGAttributeAccessorPrivate::GetAllocationSize(ElementSize, Count));
+		if (FConstructMemoryFunc ConstructFunc = GetConstructMemoryFunc(UnderlyingValueType))
+		{
+			(*ConstructFunc)(Buffer.GetData(), Count);
+		}
+
+		const bool bSuccess = GetRangeVirtual(FOutValues{TInPlaceType<FOutValuesByValue>{}, Buffer.GetData()}, Count, Index, Keys, Flags);
+			
+		if (bSuccess)
+		{
+			FTransformFunc TransformFunc = OpType == EPCGAttributeAccessorOperation::Broadcast 
+				? GetBroadcastTransformFunc(UnderlyingValueType, ExpectedType) 
+				: GetConstructibleTransformFunc(ExpectedType, UnderlyingValueType);
+			
+			check(TransformFunc);
+			(*TransformFunc)(OutValues.Get<FOutValuesByValue>().OutValues, Buffer.GetData(), Count);
+		}
+			
+		if (FDestroyMemoryFunc DestroyFunc = GetDestroyMemoryFunc(UnderlyingValueType))
+		{
+			(*DestroyFunc)(Buffer.GetData(), Count);
+		}
+			
+		return bSuccess;
+	}
+	default:
+		checkNoEntry();
+		return false;
+	}
+}
+
+bool IPCGAttributeAccessor::SetRange_Internal(PCG::Private::FInValues InValues, int16 ExpectedType, int32 Count, int32 Index, IPCGAttributeAccessorKeys& Keys, EPCGAttributeAccessorFlags Flags, PCG::Private::EPCGAttributeAccessorOperation OpType)
+{
+	using namespace PCG::Private;
+	
+	switch (OpType)
+	{
+	case EPCGAttributeAccessorOperation::SameType:
+	{
+		return SetRangeVirtual(MoveTemp(InValues), Count, Index, Keys, Flags);
+	}
+	case EPCGAttributeAccessorOperation::Construct: // Fall-through
+	case EPCGAttributeAccessorOperation::Broadcast:
+	{
+		check(InValues.IsType<FInValuesByValue>() && IsPCGType(ExpectedType))
+		
+		PCGAttributeAccessorPrivate::FRawBuffer Buffer{}; // 1KB
+		const uint16 UnderlyingValueType = static_cast<uint16>(UnderlyingDesc.ValueType);
+			
+		const int32 ElementSize = GetElementSize(UnderlyingValueType);
+		if (!ensure(ElementSize > 0))
+		{
+			return false;
+		}
+
+		Buffer.SetNumUninitialized(PCGAttributeAccessorPrivate::GetAllocationSize(ElementSize, Count));
+		if (FConstructMemoryFunc ConstructFunc = GetConstructMemoryFunc(UnderlyingValueType))
+		{
+			(*ConstructFunc)(Buffer.GetData(), Count);
+		}
+
+		FTransformFunc TransformFunc = OpType == EPCGAttributeAccessorOperation::Broadcast
+			? GetBroadcastTransformFunc(ExpectedType, UnderlyingValueType)
+			: GetConstructibleTransformFunc(UnderlyingValueType, ExpectedType);
+		
+		check(TransformFunc);
+		(*TransformFunc)(Buffer.GetData(), InValues.Get<FInValuesByValue>().InValues, Count);
+	
+		const bool bSuccess = SetRangeVirtual(FInValues{TInPlaceType<FInValuesByValue>{}, Buffer.GetData(), Count}, Count, Index, Keys, Flags);
+		
+		if (FDestroyMemoryFunc DestroyFunc = GetDestroyMemoryFunc(UnderlyingValueType))
+		{
+			(*DestroyFunc)(Buffer.GetData(), Count);
+		}
+		
+		return bSuccess;
+	}
+	default:
+		checkNoEntry();
+		return false;
+	}
+}
diff --git a/Engine/Plugins/PCG/Source/PCG/Private/Metadata/PCGMetadataAttributeGeneric.cpp b/Engine/Plugins/PCG/Source/PCG/Private/Metadata/PCGMetadataAttributeGeneric.cpp
--- a/Engine/Plugins/PCG/Source/PCG/Private/Metadata/PCGMetadataAttributeGeneric.cpp
+++ b/Engine/Plugins/PCG/Source/PCG/Private/Metadata/PCGMetadataAttributeGeneric.cpp
@@ -22,7 +22,7 @@ namespace PCG::Private
 
 	/**
 	 * Equivalent of the method in PropertyBag to convert the descriptor to a property.
-	 * 
+	 *
 	 */
 	FProperty* CreatePropertyFromDesc(const FPCGMetadataAttributeDesc& Desc, const FFieldVariant PropertyScope)
 	{
@@ -77,7 +77,7 @@ namespace PCG::Private
 							delete MapProperty;
 							return nullptr;
 						}
-						
+
 						MapProperty->KeyProp = KeyProp;
 						break;
 					}
@@ -91,7 +91,7 @@ namespace PCG::Private
 			FPCGMetadataAttributeDesc InnerDesc = Desc;
 			InnerDesc.ContainerTypes.Reset();
 			*ValuePropertyPtr = CreatePropertyFromDesc(InnerDesc, PropertyOwner);
-			
+
 			if (*ValuePropertyPtr == nullptr)
 			{
 				delete Prop;
@@ -102,9 +102,9 @@ namespace PCG::Private
 				return Prop;
 			}
 		}
-		
+
 		const UScriptStruct* ScriptStruct = nullptr;
-		
+
 		switch (Desc.ValueType)
 		{
 		case EPCGMetadataTypes::Boolean:
@@ -200,7 +200,7 @@ namespace PCG::Private
 		case EPCGMetadataTypes::Struct:
 			ScriptStruct = Cast<UScriptStruct>(Desc.ValueTypeObject);
 			break;
-			
+
 		// @todo_pcg: Support class and objects
 // 		case EPCGMetadataTypes::Class:
 // 			if (const UClass* Class = Cast<UClass>(Desc.ValueTypeObject))
@@ -259,7 +259,7 @@ namespace PCG::Private
 		default:
 			ensureMsgf(false, TEXT("Unhandled type %s"), *UEnum::GetValueAsString(Desc.ValueType));
 		}
-		
+
 		if (ScriptStruct)
 		{
 			FStructProperty* Prop = new FStructProperty(PropertyScope, Desc.Name);
@@ -269,12 +269,12 @@ namespace PCG::Private
 			{
 				Prop->SetPropertyFlags(CPF_HasGetValueTypeHash);
 			}
-				
+
 			if (ScriptStruct->StructFlags & STRUCT_HasInstancedReference)
 			{
 				Prop->SetPropertyFlags(CPF_ContainsInstancedReference);
 			}
-				
+
 			return Prop;
 		}
 
@@ -294,7 +294,7 @@ namespace PCG::Private
 		{
 			return;
 		}
-		
+
 		size_t Size = Inner->GetElementSize();
 		const uint8* SrcData = (uint8*)SrcPtr;
 		uint8* DestData = (uint8*)DestPtr;
@@ -341,14 +341,14 @@ namespace PCG::Private
 			}
 		}
 		else
-		{			
+		{
 			FMemory::Memcpy(DestData, SrcData, Count * Size);
 		}
 	}
 }
 
 /**
- * Create a new wrapper using the Attribute desc. This is expected to be stored in a SharedPtr to avoid having 
+ * Create a new wrapper using the Attribute desc. This is expected to be stored in a SharedPtr to avoid having
  * to always recreate it.
  */
 FPCGAttributeProperty::FPCGAttributeProperty(const FPCGMetadataAttributeDesc& InAttributeDesc)
@@ -366,19 +366,19 @@ FPCGAttributeProperty::FPCGAttributeProperty(const FPCGMetadataAttributeDesc& In
 		Property.Reset();
 		return;
 	}
-	
+
 	Property->AddCppProperty(NewProperty);
 	// We need to link the property so it fills the Property Flags and other members correctly.
 	// We don't need an archive (that's how it is done in the FProperty::StaticLink), but FNullArchive is private so we copied it.
 	PCG::Private::FNullArchive ArDummy;
 	Property->Link(ArDummy);
-	
+
 	// Cache info if the data is plain old data.
 	bIsPlainOldData = !!(NewProperty->PropertyFlags & CPF_IsPlainOldData);
-	
+
 	// Compress strings and soft paths, if they are not in a container.
-	bCompressData = InAttributeDesc.ContainerTypes.IsEmpty() &&  
-		(InAttributeDesc.ValueType == EPCGMetadataTypes::String 
+	bCompressData = InAttributeDesc.ContainerTypes.IsEmpty() &&
+		(InAttributeDesc.ValueType == EPCGMetadataTypes::String
 		|| InAttributeDesc.ValueType == EPCGMetadataTypes::Text
 		|| InAttributeDesc.ValueType == EPCGMetadataTypes::SoftObjectPath
 		|| InAttributeDesc.ValueType == EPCGMetadataTypes::SoftClassPath
@@ -417,7 +417,7 @@ void FPCGAttributeDefaultValue::Init(TSharedPtr<FPCGAttributeProperty> InUnderly
 	{
 		return;
 	}
-	
+
 	// Check if we already have a property set. If the property is the same type, we can keep the allocated memory.
 	// Otherwise we have to destroy it.
 	if (UnderlyingProperty && UnderlyingProperty->IsValid() && Memory)
@@ -428,9 +428,9 @@ void FPCGAttributeDefaultValue::Init(TSharedPtr<FPCGAttributeProperty> InUnderly
 			return;
 		}
 	}
-	
+
 	Destroy();
-	
+
 	UnderlyingProperty = MoveTemp(InUnderlyingProperty);
 	Memory = UnderlyingProperty->Property->Inner->AllocateAndInitializeValue();
 }
@@ -445,7 +445,7 @@ void FPCGAttributeDefaultValue::Destroy()
 {
 	if (Memory)
 	{
-		if (ensureMsgf(UnderlyingProperty && UnderlyingProperty->IsValid(), 
+		if (ensureMsgf(UnderlyingProperty && UnderlyingProperty->IsValid(),
 			TEXT("Memory is allocated but we do not have an underlying property. This is ill-formed and will leak memory")))
 		{
 			UnderlyingProperty->Property->Inner->DestroyValue(Memory);
@@ -469,7 +469,7 @@ FPCGMetadataAttributeGeneric::FPCGMetadataAttributeGeneric(const FPCGMetadataAtt
 		CachedDesc = FPCGMetadataAttributeDesc{};
 		return;
 	}
-	
+
 	Init();
 }
 
@@ -484,7 +484,7 @@ FPCGMetadataAttributeGeneric::FPCGMetadataAttributeGeneric(TNotNull<const FPCGMe
 	Init();
 
 	ValueKeyOffset = InParent->GetValueKeyOffsetForChild();
-	
+
 	if (ensure(DefaultValue.IsValid()) && ensure(InParent->DefaultValue.IsValid()))
 	{
 		// Copy the default value from the parent.
@@ -529,7 +529,7 @@ void FPCGMetadataAttributeGeneric::GetResourceSizeEx(FResourceSizeEx& Cumulative
 	FArchiveCountMem Archive(nullptr);
 	// Serialize desc
 	FPCGMetadataAttributeDesc::StaticStruct()->SerializeItem(FStructuredArchiveFromArchive(Archive).GetSlot(), const_cast<FPCGMetadataAttributeDesc*>(&CachedDesc), /*Defaults*/nullptr);
-	
+
 	if (UnderlyingProperty && UnderlyingProperty->Property)
 	{
 		if (UnderlyingProperty->Property->Inner && DefaultValue.GetRawPtr())
@@ -585,10 +585,10 @@ void FPCGMetadataAttributeGeneric::Serialize(FPCGMetadataDomain* InMetadata, FAr
 		{
 			UnderlyingProperty->Property->Inner->SerializeItem(FStructuredArchiveFromArchive(InArchive).GetSlot(), DefaultValue.GetRawPtr(), /*Defaults*/nullptr);
 		}
-	
+
 		// And finally serialize the values
 		UnderlyingProperty->Property->SerializeItem(FStructuredArchiveFromArchive(InArchive).GetSlot(), &Values, /*Defaults*/nullptr);
-	
+
 		if (InArchive.IsLoading())
 		{
 			ValueKeyOffset = GetParent() ? GetParent()->GetValueKeyOffsetForChild() : 0;
@@ -622,8 +622,10 @@ bool FPCGMetadataAttributeGeneric::IsMap() const
 	return !CachedDesc.ContainerTypes.IsEmpty() && CachedDesc.ContainerTypes[0] == EPCGMetadataAttributeContainerTypes::Map;
 }
 
-TArray<PCGMetadataValueKey> FPCGMetadataAttributeGeneric::AddValues_Internal(FInValues InValues, int32 Count, int32 ExistingIndex, bool bLockless)
+TArray<PCGMetadataValueKey> FPCGMetadataAttributeGeneric::AddValues_Internal(const PCG::Private::FInValues& InValues, int32 Count, int32 ExistingIndex, bool bLockless)
 {
+	using namespace PCG::Private;
+
 	TArray<PCGMetadataValueKey> ValueKeys;
 
 	int32 FirstValueIndex = INDEX_NONE;
@@ -633,9 +635,9 @@ TArray<PCGMetadataValueKey> FPCGMetadataAttributeGeneric::AddValues_Internal(FIn
 	{
 		return ValueKeys;
 	}
-	
+
 	// Hard-requirement that when we write to the default value, we can only write a single value
-	if (!ensure(ExistingIndex != PCG::Private::ExistingIndexForDefaultValue || Count == 1))
+	if (!ensure(ExistingIndex != ExistingIndexForDefaultValue || Count == 1))
 	{
 		Count = 1;
 	}
@@ -647,7 +649,7 @@ TArray<PCGMetadataValueKey> FPCGMetadataAttributeGeneric::AddValues_Internal(FIn
 			check(!bLockless);
 			FirstValueIndex = AllocateValues_Unsafe(Count);
 		}
-		else if (ExistingIndex != PCG::Private::ExistingIndexForDefaultValue)
+		else if (ExistingIndex != ExistingIndexForDefaultValue)
 		{
 			// Expected to be already allocated.
 			check(Helper->Num() >= ExistingIndex + Count);
@@ -664,7 +666,7 @@ TArray<PCGMetadataValueKey> FPCGMetadataAttributeGeneric::AddValues_Internal(FIn
 
 		auto GetRawPtr = [this, &Helper, FirstValueIndex](const int32 Index)
 		{
-			return FirstValueIndex == PCG::Private::ExistingIndexForDefaultValue ? DefaultValue.GetRawPtr() : Helper->GetRawPtr(Index);
+			return FirstValueIndex == ExistingIndexForDefaultValue ? DefaultValue.GetRawPtr() : Helper->GetRawPtr(Index);
 		};
 
 		Visit([this, &GetRawPtr, FirstValueIndex, Count](auto&& Value)
@@ -692,12 +694,12 @@ TArray<PCGMetadataValueKey> FPCGMetadataAttributeGeneric::AddValues_Internal(FIn
 					{
 						FScriptSetHelper OtherHelper(InnerProperty, GetRawPtr(FirstValueIndex + i));
 						const int32 Num = Value.InValues[i].Num();
-				
+
 						for (int j = 0; j < Num; ++j)
 						{
 							OtherHelper.AddElement(Value.InValues[i][j]);
 						}
-					}
+						}
 				}
 				else if constexpr (std::is_same_v<T, FInValuesAsMap>)
 				{
@@ -706,12 +708,12 @@ TArray<PCGMetadataValueKey> FPCGMetadataAttributeGeneric::AddValues_Internal(FIn
 					{
 						FScriptMapHelper OtherHelper(InnerProperty, GetRawPtr(FirstValueIndex + i));
 						const int32 Num = Value.InValues[i].Num();
-				
+
 						for (int j = 0; j < Num; ++j)
 						{
 							OtherHelper.AddPair(Value.InValues[i][j].Key, Value.InValues[i][j].Value);
 						}
-					}
+						}
 				}
 				else if constexpr (std::is_same_v<T, FInValuesAsArray>)
 				{
@@ -728,7 +730,7 @@ TArray<PCGMetadataValueKey> FPCGMetadataAttributeGeneric::AddValues_Internal(FIn
 						{
 							OtherHelper.AddValues(Num);
 						}
-				
+
 						PCG::Private::CopyArray(InnerProperty, OtherHelper.GetRawPtr(), Value.InValues[i].template Get<0>(), Num);
 					}
 				}
@@ -748,9 +750,10 @@ TArray<PCGMetadataValueKey> FPCGMetadataAttributeGeneric::AddValues_Internal(FIn
 	return ValueKeys;
 }
 
-TArray<PCGMetadataValueKey> FPCGMetadataAttributeGeneric::AddCompressedValues(FInValuesByValue InUniqueValues,
-	TArray<PCGMetadataValueKey>& FoundValueKeys)
+TArray<PCGMetadataValueKey> FPCGMetadataAttributeGeneric::AddCompressedValues(PCG::Private::FInValuesByValue InUniqueValues, TArray<PCGMetadataValueKey>& FoundValueKeys)
 {
+	using namespace PCG::Private;
+
 	const bool bHasDuplicateValues = (InUniqueValues.Count != FoundValueKeys.Num());
 
 	TArray<PCGMetadataValueKey> FoundUniqueValueKeys;
@@ -798,14 +801,14 @@ TArray<PCGMetadataValueKey> FPCGMetadataAttributeGeneric::AddCompressedValues(FI
 	return FoundValueKeys;
 }
 
-bool FPCGMetadataAttributeGeneric::FindValues(TArray<PCGMetadataValueKey>& OutValueKeys, const FInValuesByValue& InValues) const
+bool FPCGMetadataAttributeGeneric::FindValues(TArray<PCGMetadataValueKey>& OutValueKeys, const PCG::Private::FInValuesByValue& InValues) const
 {
 	OutValueKeys.Init(PCGNotFoundValueKey, InValues.Count);
-	
+
 	return FindValues(MakeArrayView(OutValueKeys), InValues);
 }
 
-bool FPCGMetadataAttributeGeneric::FindValues(TArrayView<PCGMetadataValueKey> OutValueKeys, const FInValuesByValue& InValues) const
+bool FPCGMetadataAttributeGeneric::FindValues(TArrayView<PCGMetadataValueKey> OutValueKeys, const PCG::Private::FInValuesByValue& InValues) const
 {
 	check(OutValueKeys.Num() == InValues.Count);
 	int ValueKeysSet = 0;
@@ -814,20 +817,22 @@ bool FPCGMetadataAttributeGeneric::FindValues(TArrayView<PCGMetadataValueKey> Ou
 	return (ValueKeysSet == InValues.Count);
 }
 
-void FPCGMetadataAttributeGeneric::FindValuesInternal(const FInValuesByValue& InValues, TArrayView<PCGMetadataValueKey> ValueKeys, int& ValueKeysSet, bool bIsRoot) const
+void FPCGMetadataAttributeGeneric::FindValuesInternal(const PCG::Private::FInValuesByValue& InValues, TArrayView<PCGMetadataValueKey> ValueKeys, int& ValueKeysSet, bool bIsRoot) const
 {
+	using namespace PCG::Private;
+
 	check(InValues.Count == ValueKeys.Num());
-	
+
 	if (!ensure(UnderlyingProperty))
 	{
 		return;
 	}
 
 	auto Compare = [this, &InValues](int32 InValueIndex, PCGMetadataValueKey ValueKey)
-	{	
+	{
 		const void* Value1 = GetReadAddressForValueKey_Unsafe(ValueKey);
 		const void* Value2 = UnderlyingProperty->GetPtrInArray(InValues.InValues, InValueIndex);
-	
+
 		if (Value1 == nullptr || Value2 == nullptr)
 		{
 			return false;
@@ -888,6 +893,9 @@ void FPCGMetadataAttributeGeneric::Init()
 		return;
 	}
 
+	// In order to be consistent between multiple equivalent descriptors, we will convert structs that matches a predefined type.
+	CachedDesc.FixLegacyTypeId();
+
 	// Initialize the default value
 	DefaultValue.Init(UnderlyingProperty);
 
@@ -994,11 +1002,11 @@ void FPCGMetadataAttributeGeneric::FlattenAndCompress(const TArrayView<const PCG
 	// Then for each value key (or unique values keys), gather the value in a NewValues array
 	// and also keep a mapping between old value key and new value key.
 	// Only done if the old value key is not the default one.
-	
+
 	TMap<PCGMetadataValueKey, PCGMetadataValueKey> ValueKeyMapping;
 	FScriptArray NewValues{};
 	FScriptArrayHelper NewValuesHelper{UnderlyingProperty->Property.Get(), &NewValues};
-	
+
 	const TArray<PCGMetadataValueKey>& AllValueKeysRef = !bUseValueKeys ? AllValueKeys : AllUniqueValueKeys;
 	if (UnderlyingProperty->bIsPlainOldData)
 	{
@@ -1008,16 +1016,16 @@ void FPCGMetadataAttributeGeneric::FlattenAndCompress(const TArrayView<const PCG
 	{
 		NewValuesHelper.AddValues(AllValueKeysRef.Num());
 	}
-	
+
 	ValueKeyMapping.Reserve(AllValueKeysRef.Num());
-	
+
 	int32 Count = 0;
 	for (PCGMetadataValueKey ValueKey : AllValueKeysRef)
 	{
 		if (ValueKey != PCGDefaultValueKey)
 		{
 			ValueKeyMapping.Add(ValueKey, Count);
-			
+
 			// @todo_pcg: Improve to find all contiguous values before copying.
 			const void* SrcPtr = GetReadAddressForValueKey_Unsafe(ValueKey);
 			UnderlyingProperty->Copy(NewValuesHelper.GetRawPtr(Count++), SrcPtr, /*Count=*/1);
@@ -1101,7 +1109,7 @@ FPCGMetadataAttributeBase* FPCGMetadataAttributeGeneric::Copy(FName NewName, FPC
 		NewDesc.Name = NewName;
 		AttributeCopy = new FPCGMetadataAttributeGeneric(NewDesc, InMetadata, bAllowsInterpolation);
 	}
-	
+
 	CopyInternal(AttributeCopy, bKeepParent, bCopyEntries, bCopyValues);
 
 	return AttributeCopy;
@@ -1195,7 +1203,7 @@ void FPCGMetadataAttributeGeneric::CopyInternal(FPCGMetadataAttributeGeneric* Ne
 				{
 					continue;
 				}
-				
+
 				const int32 NumValuesToCopy = ParentHelper->Num();
 				if (NumValuesToCopy <= 0)
 				{
@@ -1203,7 +1211,7 @@ void FPCGMetadataAttributeGeneric::CopyInternal(FPCGMetadataAttributeGeneric* Ne
 				}
 
 				const int32 StartIndex = NewAttribute->AllocateValues_Unsafe(NumValuesToCopy);
-				
+
 				if (StartIndex == INDEX_NONE)
 				{
 					continue;
@@ -1217,6 +1225,8 @@ void FPCGMetadataAttributeGeneric::CopyInternal(FPCGMetadataAttributeGeneric* Ne
 
 void FPCGMetadataAttributeGeneric::SetValue(PCGMetadataEntryKey ItemKey, const FPCGMetadataAttributeBase* InAttribute, PCGMetadataEntryKey InEntryKey)
 {
+	using namespace PCG::Private;
+
 	if (!InAttribute || !InAttribute->IsGeneric())
 	{
 		return;
@@ -1229,28 +1239,28 @@ void FPCGMetadataAttributeGeneric::SetValue(PCGMetadataEntryKey ItemKey, const F
 	{
 		return;
 	}
-	
+
 	check(ItemKey != PCGInvalidEntryKey);
-	
+
 	const void* SrcPtr = InAttributeGeneric->GetReadAddressFromEntryKey_Unsafe(InEntryKey);
 	if (!SrcPtr)
 	{
-		return;	
+		return;
 	}
-	
+
 	PCGMetadataValueKey ValueKeyToAdd = PCGNotFoundValueKey;
-	
+
 	if (UnderlyingProperty->bCompressData)
 	{
-		const bool bValueKeyFound = FindValues(MakeArrayView(&ValueKeyToAdd, 1), FInValuesByValue{SrcPtr, 1});
+		const bool bValueKeyFound = FindValues(MakeArrayView(&ValueKeyToAdd, 1), PCG::Private::FInValuesByValue{SrcPtr, 1});
 		check(!bValueKeyFound || ValueKeyToAdd != PCGNotFoundValueKey);
 	}
-	
+
 	// If we are not compressed or we haven't find the value, we need to add it.
 	if (ValueKeyToAdd == PCGNotFoundValueKey)
 	{
 		PCG::TUniqueScopeLock WriteLock(ValueLock);
-		
+
 		// Add the value in the buffer
 		const int32 NewValueIndex = ThisHelper->AddValue();
 
@@ -1259,7 +1269,7 @@ void FPCGMetadataAttributeGeneric::SetValue(PCGMetadataEntryKey ItemKey, const F
 		{
 			SrcPtr = InAttributeGeneric->GetReadAddressFromEntryKey_Unsafe(InEntryKey);
 		}
-		
+
 		check(SrcPtr);
 
 		ValueKeyToAdd = ValueKeyOffset + NewValueIndex;
@@ -1303,7 +1313,7 @@ bool FPCGMetadataAttributeGeneric::IsEqualToDefaultValue(PCGMetadataValueKey Val
 	{
 		return true;
 	}
-	
+
 	if (ensure(UnderlyingProperty) && UnderlyingProperty->bCompressData)
 	{
 		return false;
@@ -1328,11 +1338,11 @@ void FPCGMetadataAttributeGeneric::SetDefaultValueToFirstEntry()
 		const void* SrcPtr = GetReadAddressForValueKey_Unsafe(ValueKeyOffset);
 		if (!SrcPtr)
 		{
-			return;	
+			return;
 		}
-		
+
 		void* DestPtr = DefaultValue.GetRawPtr();
-		
+
 		UnderlyingProperty->Copy(DestPtr, SrcPtr);
 	}
 }
@@ -1353,23 +1363,23 @@ bool FPCGMetadataAttributeGeneric::AreValuesEqual(PCGMetadataValueKey ValueKey1,
 	{
 		return false;
 	}
-	
+
 	if (ValueKey1 == ValueKey2)
 	{
 		return true;
 	}
-	
+
 	{
 		PCG::TSharedScopeLock ReadLock(ValueLock);
-		
+
 		const void* Value1 = GetReadAddressForValueKey_Unsafe(ValueKey1);
 		const void* Value2 = GetReadAddressForValueKey_Unsafe(ValueKey2);
-		
+
 		if (Value1 == nullptr || Value2 == nullptr)
 		{
 			return false;
 		}
-		
+
 		return UnderlyingProperty->Property->Inner->Identical(Value1, Value2);
 	}
 }
@@ -1384,15 +1394,17 @@ bool FPCGMetadataAttributeGeneric::AreValuesEqual(PCGMetadataValueKey ValueKey1,
  * - Output as Array
  * - Output as Set
  * - Output as Map
- * 
+ *
  * On top of that, in order to reduce as much as we can virtual function calls, we also track for the different un-retrieved values if they are contiguous in memory, for input AND output.
  * Values are contiguous in memory if their index in the output array are contiguous AND the values keys associated with them are also contiguous.
  * For example, if we have indices {0, 1, 2, 5} (because we still haven't retrieved values for those indices), and we have value keys {1, 2, 8, 9}, we have a contiguous
  * range for the first two (their indices AND values keys are contiguous), but not for the rest.
  * So we can do 3 extractions (in this example [Index:ValueKey]): [0:1 ,1:2], [2:8] and [5:9] and optimize the first one to reduce the number of virtual calls.
  */
-void FPCGMetadataAttributeGeneric::GetValues_Internal(const TArrayView<const PCGMetadataValueKey> ValueKeys, FOutValues OutValues, TBitArray<>& UnretrievedValues) const
+void FPCGMetadataAttributeGeneric::GetValues_Internal(const TArrayView<const PCGMetadataValueKey> ValueKeys, PCG::Private::FOutValues OutValues, TBitArray<>& UnretrievedValues) const
 {
+	using namespace PCG::Private;
+
 	bool bFoundAllKeys = true;
 	TConstSetBitIterator<> It(UnretrievedValues);
 	if (!It)
@@ -1409,7 +1421,7 @@ void FPCGMetadataAttributeGeneric::GetValues_Internal(const TArrayView<const PCG
 
 	{
 		PCG::TSharedScopeLock ReadLock(ValueLock);
-		
+
 		struct FContinuousRange
 		{
 			void Reset()
@@ -1418,21 +1430,21 @@ void FPCGMetadataAttributeGeneric::GetValues_Internal(const TArrayView<const PCG
 				StartValueKey = INDEX_NONE;
 				Count = 0;
 			}
-			
+
 			void Init(int32 InStartIndex, int64 InStartValueKey)
 			{
 				StartIndex = InStartIndex;
 				StartValueKey = InStartValueKey;
 				Count = 1;
 			}
-			
+
 			int32 StartIndex = INDEX_NONE;
 			int64 StartValueKey = INDEX_NONE;
 			int32 Count = 0;
 		};
-		
+
 		FContinuousRange ContinuousRange{};
-		
+
 		auto RetrieveValues = [this, &OutValues, &Helper](FContinuousRange& Range)
 		{
 			const uint8* SrcPtr = nullptr;
@@ -1502,14 +1514,14 @@ void FPCGMetadataAttributeGeneric::GetValues_Internal(const TArrayView<const PCG
 					}
 			}, OutValues);
 		};
-		
+
 		for (; It; ++It)
 		{
 			const int32 Index = It.GetIndex();
 			const PCGMetadataValueKey ValueKey = ValueKeys[Index];
-			
+
 			int32 ValueKeyIndex = INDEX_NONE;
-			
+
 			if (ValueKey == PCGDefaultValueKey)
 			{
 				ValueKeyIndex = PCG::Private::ExistingIndexForDefaultValue;
@@ -1527,7 +1539,7 @@ void FPCGMetadataAttributeGeneric::GetValues_Internal(const TArrayView<const PCG
 			{
 				bFoundAllKeys = false;
 			}
-			
+
 			if (ValueKeyIndex < 0)
 			{
 				// Reached the end
@@ -1536,7 +1548,7 @@ void FPCGMetadataAttributeGeneric::GetValues_Internal(const TArrayView<const PCG
 					RetrieveValues(ContinuousRange);
 					ContinuousRange.Reset();
 				}
-				
+
 				// Then copy the default value if the value key is PCG::Private::ExistingIndexForDefaultValue
 				if (ValueKeyIndex == PCG::Private::ExistingIndexForDefaultValue)
 				{
@@ -1560,17 +1572,17 @@ void FPCGMetadataAttributeGeneric::GetValues_Internal(const TArrayView<const PCG
 					RetrieveValues(ContinuousRange);
 					ContinuousRange.Reset();
 				}
-				
+
 				ContinuousRange.Init(Index, ValueKeyIndex);
 			}
-			
+
 			// If the value key index is not none, the value was (or will be) retrieved.
 			if (ValueKeyIndex != INDEX_NONE)
 			{
 				UnretrievedValues[Index] = false;
 			}
 		}
-		
+
 		// One last time
 		if (ContinuousRange.StartIndex != INDEX_NONE)
 		{
@@ -1594,18 +1606,18 @@ FInstancedPropertyBag FPCGMetadataAttributeGeneric::BuildStructForDebug(PCGMetad
 	{
 		return {};
 	}
-	
+
 	FInstancedPropertyBag Result;
-	
+
 	FPropertyBagPropertyDesc Desc{Name, UnderlyingProperty->Property->Inner};
-	
+
 	const UPropertyBag* PropertyBag = UPropertyBag::GetOrCreateFromDescs(MakeConstArrayView(&Desc, 1));
 	Result.InitializeFromBagStruct(PropertyBag);
-	
+
 	const FProperty* ResultProperty = PropertyBag->GetPropertyDescs()[0].CachedProperty;
-	
+
 	void* DestPtr = ResultProperty->ContainerPtrToValuePtr<void>(Result.GetMutableValue().GetMemory());
-	
+
 	{
 		PCG::TSharedScopeLock ReadLock(ValueLock);
 		if (const void* SrcPtr = GetReadAddressFromEntryKey_Unsafe(EntryKey))
@@ -1623,7 +1635,7 @@ TOptional<FScriptArrayHelper> FPCGMetadataAttributeGeneric::MakeArrayHelper() co
 	{
 		return {};
 	}
-	
+
 	return FScriptArrayHelper(UnderlyingProperty->Property.Get(), &Values);
 }
 
@@ -1641,7 +1653,7 @@ const void* FPCGMetadataAttributeGeneric::GetReadAddressForValueKey_Unsafe(const
 	}
 
 	const int32 ValueIndex = InValueKey - ValueKeyOffset;
-	
+
 	TOptional<FScriptArrayHelper> Helper = MakeArrayHelper();
 	if (!Helper || ValueIndex >= Values.Num())
 	{
@@ -1658,22 +1670,24 @@ const void* FPCGMetadataAttributeGeneric::GetReadAddressFromEntryKey_Unsafe(cons
 
 void FPCGMetadataAttributeGeneric::SetValueFromProperty(PCGMetadataEntryKey ItemKey, const void* SrcPtr, const FProperty* Property)
 {
+	using namespace PCG::Private;
+
 	if (!SrcPtr || !Property)
 	{
 		return;
 	}
-	
+
 	if (TOptional<FScriptArrayHelper> Helper = MakeArrayHelper(); Helper.IsSet())
 	{
 		if (!UnderlyingProperty->Property->Inner->SameType(Property))
 		{
 			return;
 		}
-		
+
 		const int32 ExistingIndex = ItemKey == PCGInvalidEntryKey && ensure(DefaultValue.IsValid()) ? PCG::Private::ExistingIndexForDefaultValue : INDEX_NONE;
-		
+
 		PCGMetadataValueKey ValueKeyToAdd = PCGNotFoundValueKey;
-	
+
 		if (UnderlyingProperty->bCompressData && ItemKey != PCGInvalidEntryKey)
 		{
 			const bool bValueKeyFound = FindValues(MakeArrayView(&ValueKeyToAdd, 1), FInValuesByValue{SrcPtr, 1});
@@ -1694,7 +1708,7 @@ void FPCGMetadataAttributeGeneric::SetValueFromProperty(PCGMetadataEntryKey Item
 int32 FPCGMetadataAttributeGeneric::PreallocateValues(TArrayView<PCGMetadataEntryKey*> EntryKeys, bool bLockless)
 {
 	int32 StartIndex = INDEX_NONE;
-	
+
 	if (!ensure(UnderlyingProperty))
 	{
 		return StartIndex;
diff --git a/Engine/Plugins/PCG/Source/PCG/Private/Metadata/PCGMetadataCommon.cpp b/Engine/Plugins/PCG/Source/PCG/Private/Metadata/PCGMetadataCommon.cpp
--- a/Engine/Plugins/PCG/Source/PCG/Private/Metadata/PCGMetadataCommon.cpp
+++ b/Engine/Plugins/PCG/Source/PCG/Private/Metadata/PCGMetadataCommon.cpp
@@ -2,6 +2,8 @@
 
 #include "Metadata/PCGMetadataCommon.h"
 
+#include "UObject/Class.h"
+
 #include UE_INLINE_GENERATED_CPP_BY_NAME(PCGMetadataCommon)
 
 namespace PCGMetadataDomainID
@@ -26,49 +28,55 @@ bool FPCGMetadataAttributeDesc::operator==(const FPCGMetadataAttributeDesc& Othe
 	return Name == Other.Name && IsSameType(Other);
 }
 
-int16 FPCGMetadataAttributeDesc::AsLegacyTypeId() const
+void FPCGMetadataAttributeDesc::FixLegacyTypeId()
 {
-	if (static_cast<uint8>(ValueType) < static_cast<uint8>(EPCGMetadataTypes::EndLegacyTypes))
-	{
-		return static_cast<int16>(ValueType);
-	}
-	else if (ValueType == EPCGMetadataTypes::Struct)
+	auto FixType = [](EPCGMetadataTypes& InOutType, TObjectPtr<const UObject>& InOutTypeObject)
 	{
-		if (ValueTypeObject == TBaseStructure<FVector2D>::Get())
+		if (InOutType == EPCGMetadataTypes::Struct)
 		{
-			return static_cast<int16>(EPCGMetadataTypes::Vector2);
+			auto AssignType = [&InOutType, &InOutTypeObject](EPCGMetadataTypes NewType)
+			{
+				InOutType = NewType;
+				InOutTypeObject = nullptr;
+			};
+	
+			if (InOutTypeObject == TBaseStructure<FVector2D>::Get())
+			{
+				AssignType(EPCGMetadataTypes::Vector2);
+			}
+			else if (InOutTypeObject == TBaseStructure<FVector>::Get())
+			{
+				AssignType(EPCGMetadataTypes::Vector);
+			}
+			else if (InOutTypeObject == TBaseStructure<FVector4>::Get())
+			{
+				AssignType(EPCGMetadataTypes::Vector4);
+			}
+			else if (InOutTypeObject == TBaseStructure<FRotator>::Get())
+			{
+				AssignType(EPCGMetadataTypes::Rotator);
+			}
+			else if (InOutTypeObject == TBaseStructure<FQuat>::Get())
+			{
+				AssignType(EPCGMetadataTypes::Quaternion);
+			}
+			else if (InOutTypeObject == TBaseStructure<FTransform>::Get())
+			{
+				AssignType(EPCGMetadataTypes::Transform);
+			}
+			else if (InOutTypeObject == TBaseStructure<FSoftObjectPath>::Get())
+			{
+				AssignType(EPCGMetadataTypes::SoftObjectPath);
+			}
+			else if (InOutTypeObject == TBaseStructure<FSoftClassPath>::Get())
+			{
+				AssignType(EPCGMetadataTypes::SoftClassPath);
+			}
 		}
-		else if (ValueTypeObject == TBaseStructure<FVector>::Get())
-		{
-			return static_cast<int16>(EPCGMetadataTypes::Vector);
-		}
-		else if (ValueTypeObject == TBaseStructure<FVector4>::Get())
-		{
-			return static_cast<int16>(EPCGMetadataTypes::Vector4);
-		}
-		else if (ValueTypeObject == TBaseStructure<FRotator>::Get())
-		{
-			return static_cast<int16>(EPCGMetadataTypes::Rotator);
-		}
-		else if (ValueTypeObject == TBaseStructure<FQuat>::Get())
-		{
-			return static_cast<int16>(EPCGMetadataTypes::Quaternion);
-		}
-		else if (ValueTypeObject == TBaseStructure<FTransform>::Get())
-		{
-			return static_cast<int16>(EPCGMetadataTypes::Transform);
-		}
-		else if (ValueTypeObject == TBaseStructure<FSoftObjectPath>::Get())
-		{
-			return static_cast<int16>(EPCGMetadataTypes::SoftObjectPath);
-		}
-		else if (ValueTypeObject == TBaseStructure<FSoftClassPath>::Get())
-		{
-			return static_cast<int16>(EPCGMetadataTypes::SoftClassPath);
-		}
-	}
-
-	return  static_cast<int16>(EPCGMetadataTypes::Unknown); 
+	};
+	
+	FixType(ValueType, ValueTypeObject);
+	FixType(KeyType, KeyTypeObject);
 }
 
 int32 GetTypeHash(const FPCGMetadataAttributeDesc& Desc)
diff --git a/Engine/Plugins/PCG/Source/PCG/Private/Metadata/PCGMetadataDomain.cpp b/Engine/Plugins/PCG/Source/PCG/Private/Metadata/PCGMetadataDomain.cpp
--- a/Engine/Plugins/PCG/Source/PCG/Private/Metadata/PCGMetadataDomain.cpp
+++ b/Engine/Plugins/PCG/Source/PCG/Private/Metadata/PCGMetadataDomain.cpp
@@ -236,22 +236,19 @@ FPCGMetadataAttributeGeneric* FPCGMetadataDomain::CreateGenericAttribute(const F
 		return nullptr;
 	}
 
-	const FPCGMetadataAttributeBase* ParentAttribute = nullptr;
+	const FPCGMetadataAttributeGeneric* ParentAttribute = nullptr;
 
 	if (bOverrideParent && Parent)
 	{
-		ParentAttribute = Parent->GetConstAttribute(AttributeName);
+		const FPCGMetadataAttributeBase* BaseParentAttribute = Parent->GetConstAttribute(AttributeName);
+		ParentAttribute = BaseParentAttribute && BaseParentAttribute->IsGeneric() ? static_cast<const FPCGMetadataAttributeGeneric*>(BaseParentAttribute) : nullptr;
 	}
 
 	FPCGMetadataAttributeGeneric* NewAttribute = nullptr;
 
-	PRAGMA_DISABLE_DEPRECATION_WARNINGS
-	int16 TypeId = Desc.AsLegacyTypeId();
-	PRAGMA_ENABLE_DEPRECATION_WARNINGS
-
-	if (!ParentAttribute || (ParentAttribute->GetTypeId() != TypeId || !ParentAttribute->IsGeneric()))
+	if (!ParentAttribute || (!ParentAttribute->GetAttributeDesc().IsSameType(Desc)))
 	{
-		// Can't parent if the types do not match, or is not generic
+		// Can't parent if the types do not match
 		NewAttribute = new FPCGMetadataAttributeGeneric(Desc, this, bAllowsInterpolation);
 	}
 	else
@@ -262,10 +259,12 @@ FPCGMetadataAttributeGeneric* FPCGMetadataDomain::CreateGenericAttribute(const F
 	{
 		PCG::TUniqueScopeLock WriteLock(AttributeLock);
 
-		if (FPCGMetadataAttributeBase** ExistingAttribute = Attributes.Find(AttributeName))
+		if (FPCGMetadataAttributeBase** BaseExistingAttribute = Attributes.Find(AttributeName))
 		{
+			FPCGMetadataAttributeGeneric* ExistingAttribute = *BaseExistingAttribute && (*BaseExistingAttribute)->IsGeneric() ? static_cast<FPCGMetadataAttributeGeneric*>(*BaseExistingAttribute) : nullptr;
+			
 			delete NewAttribute;
-			if ((*ExistingAttribute)->GetTypeId() != TypeId || !(*ExistingAttribute)->IsGeneric())
+			if (!ExistingAttribute || !ExistingAttribute->GetAttributeDesc().IsSameType(Desc))
 			{
 				UE_LOG(LogPCG, Error, TEXT("Attribute %s already exists but is not the right type. Abort."), *AttributeName.ToString());
 				return nullptr;
@@ -273,7 +272,7 @@ FPCGMetadataAttributeGeneric* FPCGMetadataDomain::CreateGenericAttribute(const F
 			else
 			{
 				UE_LOG(LogPCG, Warning, TEXT("Attribute %s already exists"), *AttributeName.ToString());
-				NewAttribute = static_cast<FPCGMetadataAttributeGeneric*>(*ExistingAttribute);
+				NewAttribute = ExistingAttribute;
 			}
 		}
 		else
diff --git a/Engine/Plugins/PCG/Source/PCG/Private/Tests/Accessors/PCGAttributeAccessorTest.cpp b/Engine/Plugins/PCG/Source/PCG/Private/Tests/Accessors/PCGAttributeAccessorTest.cpp
deleted file mode 100644
diff --git a/Engine/Plugins/PCG/Source/PCG/Public/Helpers/PCGConcepts.h b/Engine/Plugins/PCG/Source/PCG/Public/Helpers/PCGConcepts.h
new file mode 100644
--- /dev/null
+++ b/Engine/Plugins/PCG/Source/PCG/Public/Helpers/PCGConcepts.h
@@ -0,0 +1,36 @@
+﻿// Copyright Epic Games, Inc. All Rights Reserved.
+
+#pragma once
+
+#include "Containers/Array.h"
+#include "Containers/ArrayView.h"
+#include "Containers/StaticArray.h"
+#include "UObject/Class.h"
+
+namespace PCG::Concepts
+{
+	template <typename T>
+	struct TIsStaticArray : std::false_type {};
+	
+	template <typename T, uint32 ...Extra>
+	struct TIsStaticArray<TStaticArray<T, Extra...>> : std::true_type {};
+	
+	template <typename Container>
+	concept CIsArrayLike = 
+		TIsTArray_V<Container> 
+		|| TIsTArrayView_V<Container>
+		|| TIsStaticArray<Container>::value;
+
+	template <typename LhsType, typename RhsType>
+	concept CIsComparable = requires(const LhsType& Lhs, const RhsType& Rhs)
+	{
+		{ Lhs == Rhs } -> UE::CSameAs<bool>;
+	};
+	
+	// Concept to catch all "basic types" that are also structs.
+	template <typename T>
+	concept CIsBasicStruct = requires (T)
+	{
+		TBaseStructure<T>::Get();
+	};
+}
diff --git a/Engine/Plugins/PCG/Source/PCG/Public/Metadata/Accessors/IPCGAttributeAccessor.h b/Engine/Plugins/PCG/Source/PCG/Public/Metadata/Accessors/IPCGAttributeAccessor.h
--- a/Engine/Plugins/PCG/Source/PCG/Public/Metadata/Accessors/IPCGAttributeAccessor.h
+++ b/Engine/Plugins/PCG/Source/PCG/Public/Metadata/Accessors/IPCGAttributeAccessor.h
@@ -3,7 +3,8 @@
 #pragma once
 
 #include "Metadata/PCGMetadataAttributeTraits.h"
-
+#include "Metadata/PCGMetadataContainerTypes.h"
+#include "Metadata/Accessors/PCGAttributeAccessorKeys.h"
 
 #include "IPCGAttributeAccessor.generated.h"
 
@@ -36,6 +37,55 @@ enum class EPCGAttributeAccessorFlags
 };
 ENUM_CLASS_FLAGS(EPCGAttributeAccessorFlags);
 
+
+namespace PCG::Private
+{
+	enum class EPCGAttributeAccessorOperation
+	{
+		SameType,
+		Construct,
+		Broadcast,
+		Invalid
+	};
+	
+	template <typename T, bool bGet>
+	EPCGAttributeAccessorOperation GetOpType(const FPCGMetadataAttributeDesc& AccessorDesc, EPCGAttributeAccessorFlags Flags)
+	{
+		const FPCGMetadataAttributeDesc* FirstDesc = &AccessorDesc;
+		const FPCGMetadataAttributeDesc* SecondDesc = &GetDefaultAttributeDesc<T>();
+		if constexpr (!bGet)
+		{
+			// Reverse order
+			Swap(FirstDesc, SecondDesc);
+		}
+		
+		if (FirstDesc->IsSameType(*SecondDesc))
+		{
+			return EPCGAttributeAccessorOperation::SameType;
+		}
+		else if (!!(Flags & EPCGAttributeAccessorFlags::AllowConstructible) && PCG::Private::IsConstructible(*SecondDesc, *FirstDesc))
+		{
+			return EPCGAttributeAccessorOperation::Construct;
+		}
+		else if (!!(Flags & EPCGAttributeAccessorFlags::AllowBroadcast) && PCG::Private::IsBroadcastable(*FirstDesc, *SecondDesc))
+		{
+			return EPCGAttributeAccessorOperation::Broadcast;
+		}
+		
+		if constexpr (bGet && std::is_same_v<T, FString>)
+		{
+			// Special case - always allow broadcasting from soft object path to string, so that legacy code that grabbed soft path attributes
+			// as FStrings will still work.
+			if (IsEquivalentDesc<FSoftObjectPath>(AccessorDesc) || IsEquivalentDesc<FSoftClassPath>(AccessorDesc))
+			{
+				return EPCGAttributeAccessorOperation::Broadcast;
+			}
+		}
+
+		return EPCGAttributeAccessorOperation::Invalid;
+	}
+}
+
 //////////////////////////////////////////////////////////////////////////////
 
 /**
@@ -63,7 +113,6 @@ public:
 	template <typename T>
 	bool Get(T& OutValue, int32 Index, const IPCGAttributeAccessorKeys& Keys, EPCGAttributeAccessorFlags Flags = EPCGAttributeAccessorFlags::StrictType) const
 	{
-		TArrayView<T> Temp(&OutValue, 1);
 		return GetRange(TArrayView<T>(&OutValue, 1), Index, Keys, Flags);
 	}
 
@@ -80,21 +129,6 @@ public:
 		return Get(OutValue, 0, Keys, Flags);
 	}
 
-	/**
-	* Get a range of values from the accessor for a given type. Not threadsafe for all accessors.
-	* If the number of elements asked is greater that the number of keys, it will wrap around.
-	* @param OutValues - View on memory where to write values to. Its "Num" will determine how many elements we will read.
-	* @param Index - The index to start looking for in the keys
-	* @param Keys - Identification to know how to retrieve the value
-	* @param Flags - Optional flag to allow for specific operations. Cf EPCGAttributeAccessorFlags.
-	* @return true if the get succeeded, false otherwise
-	*/
-	template <typename T>
-	bool GetRange(TArrayView<T> OutValues, int32 Index, const IPCGAttributeAccessorKeys& Keys, EPCGAttributeAccessorFlags Flags = EPCGAttributeAccessorFlags::StrictType) const
-	{
-		return false;
-	}
-
 	/**
 	* Set a value to the accessor for a given type. Not threadsafe for all accessors. 
 	* If Index is greater than the number of keys, it will fail (return false).
@@ -123,6 +157,37 @@ public:
 		return Set(InValue, 0, Keys, Flags);
 	}
 
+	virtual void Prepare(IPCGAttributeAccessorKeys& Keys, int32 Count, const bool bCanReuseEntryKeys) {}
+
+	int16 GetUnderlyingType() const { return UnderlyingType; }
+	const FPCGMetadataAttributeDesc& GetUnderlyingDesc() const { return UnderlyingDesc; }
+	bool IsReadOnly() const { return bReadOnly; }
+
+	// To know if we can do default value operations
+	virtual bool IsAttribute() const { return false; }
+	
+	/**
+	* Get a range of values from the accessor for a given type. Not threadsafe for all accessors.
+	* If the number of elements asked is greater that the number of keys, it will wrap around.
+	* @param OutValues - View on memory where to write values to. Its "Num" will determine how many elements we will read.
+	* @param Index - The index to start looking for in the keys
+	* @param Keys - Identification to know how to retrieve the value
+	* @param Flags - Optional flag to allow for specific operations. Cf EPCGAttributeAccessorFlags.
+	* @return true if the get succeeded, false otherwise
+	*/
+	template <typename T>
+	bool GetRange(TArrayView<T> OutValues, int32 Index, const IPCGAttributeAccessorKeys& Keys, EPCGAttributeAccessorFlags Flags = EPCGAttributeAccessorFlags::StrictType) const
+	{
+		// Have a layer of indirection to first resolve T within the view then match it against the different versions.
+		// For example:
+		// ```
+		// TArray<const FString*> Values;
+		// GetRange<const FString*>(Values, Index, Keys, Flags).
+		// ```
+		// First T resolve will resolve to `const FString*`, then GetRange_Internal will match the pointer version.
+		return GetRange_Internal(OutValues, Index, Keys, Flags);
+	}
+	
 	/**
 	* Set a range of values to the accessor for a given type. Not threadsafe for all accessors.
 	* If the number of elements asked is greater that the number of keys, it will fail. It is to avoid writing at the same memory place multiple times.
@@ -135,37 +200,306 @@ public:
 	template <typename T>
 	bool SetRange(TArrayView<const T> InValues, int32 Index, IPCGAttributeAccessorKeys& Keys, EPCGAttributeAccessorFlags Flags = EPCGAttributeAccessorFlags::StrictType)
 	{
-		return false;
+		// Can't set if read only, and discard any set that goes beyond the number of keys.
+		if (bReadOnly || Index >= Keys.GetNum())
+		{
+			return false;
+		}
+		
+		// Have a layer of indirection to first resolve T within the view then match it against the different versions.
+		// For example:
+		// ```
+		// TArray<TArray<FString>> Values;
+		// SetRange<TArray<FString>>(Values, Index, Keys, Flags).
+		// ```
+		// First T resolve will resolve to `TArray<FString>`, then SetRange_Internal will match the array version.
+		return SetRange_Internal(InValues, Index, Keys, Flags);
 	}
 
-	virtual void Prepare(IPCGAttributeAccessorKeys& Keys, int32 Count, const bool bCanReuseEntryKeys) {}
-
-	int16 GetUnderlyingType() const { return UnderlyingType; }
-	bool IsReadOnly() const { return bReadOnly; }
-
-	// To know if we can do default value operations
-	virtual bool IsAttribute() const { return false; }
-
 protected:
 	IPCGAttributeAccessor(bool bInReadOnly, int16 InUnderlyingType)
 		: bReadOnly(bInReadOnly)
 		, UnderlyingType(InUnderlyingType)
-	{}
+		, UnderlyingDesc()
+	{
+		// Should be set by the child class, but have a catch all here.
+		UnderlyingDesc.ValueType = static_cast<EPCGMetadataTypes>(UnderlyingType);
+	}
+	
+	template <typename T> requires (!std::is_pointer_v<T>)
+	bool GetRange_Internal(TArrayView<T> OutValues, int32 Index, const IPCGAttributeAccessorKeys& Keys, EPCGAttributeAccessorFlags Flags) const
+	{
+		PCG::Private::EPCGAttributeAccessorOperation OpType = PCG::Private::GetOpType<T, /*bGet=*/true>(UnderlyingDesc, Flags);
+		
+		if (OpType == PCG::Private::EPCGAttributeAccessorOperation::Invalid)
+		{
+			return false;
+		}
+		
+		return GetRange_Internal(
+			PCG::Private::FOutValues{TInPlaceType<PCG::Private::FOutValuesByValue>{}, OutValues.GetData()},
+			static_cast<int16>(PCG::Private::MetadataTypes<T>::Id),
+			OutValues.Num(),
+			Index,
+			Keys,
+			Flags,
+			OpType
+		);
+	}
+
+	// Pointer version
+	template <typename T>
+	bool GetRange_Internal(TArrayView<const T*> OutValues, int32 Index, const IPCGAttributeAccessorKeys& Keys, EPCGAttributeAccessorFlags Flags) const
+	{
+		PCG::Private::EPCGAttributeAccessorOperation OpType = PCG::Private::GetOpType<T, /*bGet=*/true>(UnderlyingDesc, Flags);
+		
+		if (OpType != PCG::Private::EPCGAttributeAccessorOperation::SameType)
+		{
+			return false;
+		}
+		
+		// We need to type erase the pointers, but to be fully C++ compliant about aliasing rules we can't reinterpret_cast. So make a copy of the pointers.
+		TArray<const void*, TInlineAllocator<256>> Temp;
+		Temp.SetNumUninitialized(OutValues.Num());
+		const bool bSuccess = GetRange_Internal(
+			PCG::Private::FOutValues{ TInPlaceType<PCG::Private::FOutValuesByPtr>{}, Temp },
+			static_cast<int16>(PCG::Private::MetadataTypes<T>::Id),
+			OutValues.Num(),
+			Index,
+			Keys,
+			Flags,
+			OpType
+		);
+		
+		if (bSuccess)
+		{
+			for (int32 i = 0; i < OutValues.Num(); ++i)
+			{
+				OutValues[i] = static_cast<const T*>(Temp[i]);
+			}
+		}
+		
+		return bSuccess;
+	}
+	
+	// Array version
+	template <typename T> 
+	bool GetRange_Internal(TArrayView<TConstArrayView<T>> OutValues, int32 Index, const IPCGAttributeAccessorKeys& Keys, EPCGAttributeAccessorFlags Flags) const
+	{
+		PCG::Private::EPCGAttributeAccessorOperation OpType = PCG::Private::GetOpType<TArray<T>, /*bGet=*/true>(UnderlyingDesc, Flags);
+		
+		if (OpType != PCG::Private::EPCGAttributeAccessorOperation::SameType)
+		{
+			return false;
+		}
+		
+		TArray<TTuple<const void*, int32>, TInlineAllocator<256>> Temp;
+		Temp.SetNumUninitialized(OutValues.Num());
+		const bool bSuccess = GetRange_Internal(
+			PCG::Private::FOutValues{TInPlaceType<PCG::Private::FOutValuesAsArray>{}, Temp},
+			static_cast<int16>(PCG::Private::MetadataTypes<T>::Id),
+			OutValues.Num(),
+			Index,
+			Keys,
+			Flags,
+			OpType
+		);
+		
+		if (bSuccess)
+		{
+			for (int32 i = 0; i < OutValues.Num(); ++i)
+			{
+				OutValues[i] = MakeConstArrayView<T>(static_cast<const T*>(Temp[i].Get<0>()), Temp[i].Get<1>());
+			}
+		}
+		
+		return bSuccess;
+	}
+	
+	// Sets version
+	template <typename T> 
+	bool GetRange_Internal(TArrayView<PCG::TScriptSetWrapper<T>> OutValues, int32 Index, const IPCGAttributeAccessorKeys& Keys, EPCGAttributeAccessorFlags Flags) const
+	{
+		PCG::Private::EPCGAttributeAccessorOperation OpType = PCG::Private::GetOpType<TSet<T>, /*bGet=*/true>(UnderlyingDesc, Flags);
+		
+		if (OpType != PCG::Private::EPCGAttributeAccessorOperation::SameType)
+		{
+			return false;
+		}
+		
+		TArray<FScriptSetHelper*, TInlineAllocator<256>> Temp;
+		Temp.Reserve(OutValues.Num());
+		for (PCG::TScriptSetWrapper<T>& Value : OutValues)
+		{
+			Temp.Add(&Value.Helper);
+		}
+
+		return GetRange_Internal(
+			PCG::Private::FOutValues{TInPlaceType<PCG::Private::FOutValuesAsSet>{}, Temp},
+			static_cast<int16>(PCG::Private::MetadataTypes<T>::Id),
+			OutValues.Num(),
+			Index,
+			Keys,
+			Flags,
+			OpType
+		);
+	}
+	
+	// Maps version
+	template <typename KeyType, typename ValueType> 
+	bool GetRange_Internal(TArrayView<PCG::TScriptMapWrapper<KeyType, ValueType>> OutValues, int32 Index, const IPCGAttributeAccessorKeys& Keys, EPCGAttributeAccessorFlags Flags) const
+	{
+		PCG::Private::EPCGAttributeAccessorOperation OpType = PCG::Private::GetOpType<TMap<KeyType, ValueType>, /*bGet=*/true>(UnderlyingDesc, Flags);
+		
+		if (OpType != PCG::Private::EPCGAttributeAccessorOperation::SameType)
+		{
+			return false;
+		}
+		
+		TArray<FScriptMapHelper*, TInlineAllocator<256>> Temp;
+		Temp.Reserve(OutValues.Num());
+		for (PCG::TScriptMapWrapper<KeyType, ValueType>& Value : OutValues)
+		{
+			Temp.Add(&Value.Helper);
+		}
+
+		return GetRange_Internal(
+			PCG::Private::FOutValues{TInPlaceType<PCG::Private::FOutValuesAsMap>{}, Temp},
+			static_cast<int16>(PCG::Private::MetadataTypes<ValueType>::Id),
+			OutValues.Num(),
+			Index,
+			Keys,
+			Flags,
+			OpType
+		);
+	}
+	
+	template <typename T>
+	bool SetRange_Internal(TArrayView<const T> InValues, int32 Index, IPCGAttributeAccessorKeys& Keys, EPCGAttributeAccessorFlags Flags = EPCGAttributeAccessorFlags::StrictType)
+	{
+		PCG::Private::EPCGAttributeAccessorOperation OpType = PCG::Private::GetOpType<T, /*bGet=*/false>(UnderlyingDesc, Flags);
+		
+		if (OpType == PCG::Private::EPCGAttributeAccessorOperation::Invalid)
+		{
+			return false;
+		}
 
-	// Define all virtual calls for all supported types.
-#define IACCESSOR_DECL(T) \
-	virtual bool GetRange##T(TArrayView<T>& OutValues, int32 Index, const IPCGAttributeAccessorKeys& Keys, EPCGAttributeAccessorFlags Flags) const = 0; \
-	virtual bool SetRange##T(TArrayView<const T>& InValues, int32 Index, IPCGAttributeAccessorKeys& Keys, EPCGAttributeAccessorFlags Flags) = 0;
-	PCG_FOREACH_SUPPORTEDTYPES(IACCESSOR_DECL);
-#undef IACCESSOR_DECL
+		return SetRange_Internal(
+			PCG::Private::FInValues{TInPlaceType<PCG::Private::FInValuesByValue>{}, InValues.GetData(), InValues.Num()},
+			static_cast<int16>(PCG::Private::MetadataTypes<T>::Id),
+			InValues.Num(),
+			Index,
+			Keys,
+			Flags,
+			OpType
+		);
+	}
+	
+	/**
+	 * Version that accept anything that is array like, like `TArray<T>`, `TArrayView<T>`, `TStaticArray<T, 5>`, `TArray<T, TInlineAllocator<5>>`, etc...
+	 */
+	template <typename ArrayContainer> requires (PCG::Concepts::CIsArrayLike<ArrayContainer>)
+	bool SetRange_Internal(TConstArrayView<ArrayContainer> InValues, int32 Index, IPCGAttributeAccessorKeys& Keys, EPCGAttributeAccessorFlags Flags)
+	{
+		using T = typename ArrayContainer::ElementType;
+		
+		PCG::Private::EPCGAttributeAccessorOperation OpType = PCG::Private::GetOpType<TArray<T>, /*bGet=*/false>(UnderlyingDesc, Flags);
+		
+		if (OpType != PCG::Private::EPCGAttributeAccessorOperation::SameType)
+		{
+			return false;
+		}
+		
+		TArray<TTuple<const void*, int32>, TInlineAllocator<256>> InValuesTuple;
+		InValuesTuple.Reserve(InValues.Num());
+		Algo::Transform(InValues, InValuesTuple, [](const auto& Item) { return TTuple<const void*, int32>{Item.GetData(), Item.Num()}; });
+
+		return SetRange_Internal(
+			PCG::Private::FInValues{TInPlaceType<PCG::Private::FInValuesAsArray>{}, InValuesTuple},
+			static_cast<int16>(PCG::Private::MetadataTypes<T>::Id),
+			InValues.Num(),
+			Index,
+			Keys,
+			Flags,
+			OpType
+		);
+	}
+
+	/**
+	 * Version that accepts Sets
+	 */
+	template <typename T, typename ...Extra>
+	bool SetRange_Internal(TConstArrayView<TSet<T, Extra...>> InValues, int32 Index, IPCGAttributeAccessorKeys& Keys, EPCGAttributeAccessorFlags Flags)
+	{
+		PCG::Private::EPCGAttributeAccessorOperation OpType = PCG::Private::GetOpType<TSet<T>, /*bGet=*/false>(UnderlyingDesc, Flags);
+		
+		if (OpType != PCG::Private::EPCGAttributeAccessorOperation::SameType)
+		{
+			return false;
+		}
+		
+		TArray<TArray<const void*>, TInlineAllocator<256>> InSetValues;
+		InSetValues.Reserve(InValues.Num());
+		Algo::Transform(InValues, InSetValues, [](const auto& Item) -> TArray<const void*>
+		{
+			TArray<const void*> SetValues;
+			SetValues.Reserve(Item.Num());
+			Algo::Transform(Item, SetValues, [](const T& Value) -> const T* { return &Value;});
+			return SetValues;
+		});
+		
+		return SetRange_Internal(
+			PCG::Private::FInValues{TInPlaceType<PCG::Private::FInValuesAsSet>{}, InSetValues},
+			static_cast<int16>(PCG::Private::MetadataTypes<T>::Id),
+			InValues.Num(),
+			Index,
+			Keys,
+			Flags,
+			OpType
+		);
+	}
+	
+	/**
+	 * Version that accepts Maps
+	 */
+	template <typename KeyType, typename ValueType, typename ...Extra>
+	bool SetRange_Internal(TConstArrayView<TMap<KeyType, ValueType, Extra...>> InValues, int32 Index, IPCGAttributeAccessorKeys& Keys, EPCGAttributeAccessorFlags Flags)
+	{
+		PCG::Private::EPCGAttributeAccessorOperation OpType = PCG::Private::GetOpType<TMap<KeyType, ValueType>, /*bGet=*/false>(UnderlyingDesc, Flags);
+		
+		if (OpType != PCG::Private::EPCGAttributeAccessorOperation::SameType)
+		{
+			return false;
+		}
+		
+		TArray<TArray<TPair<const void*, const void*>>, TInlineAllocator<256>> InMapValues;
+		InMapValues.Reserve(InValues.Num());
+		Algo::Transform(InValues, InMapValues, [](const auto& Item) -> TArray<TPair<const void*, const void*>>
+		{
+			TArray<TPair<const void*, const void*>> MapValues;
+			MapValues.Reserve(Item.Num());
+			Algo::Transform(Item, MapValues, [](const auto& It) -> TPair<const void*, const void*> { return {&It.Key, &It.Value};});
+			return MapValues;
+		});
+		
+		return SetRange_Internal(
+			PCG::Private::FInValues{TInPlaceType<PCG::Private::FInValuesAsMap>{}, InMapValues},
+			static_cast<int16>(PCG::Private::MetadataTypes<ValueType>::Id),
+			InValues.Num(),
+			Index,
+			Keys,
+			Flags,
+			OpType
+		);
+	}
+	
+	PCG_API bool GetRange_Internal(PCG::Private::FOutValues OutValues, int16 ExpectedType, int32 Count, int32 Index, const IPCGAttributeAccessorKeys& Keys, EPCGAttributeAccessorFlags Flags, PCG::Private::EPCGAttributeAccessorOperation OpType) const;
+	virtual bool GetRangeVirtual(PCG::Private::FOutValues OutValues, int32 Count, int32 Index, const IPCGAttributeAccessorKeys& Keys, EPCGAttributeAccessorFlags Flags) const = 0;
+	
+	PCG_API bool SetRange_Internal(PCG::Private::FInValues InValues, int16 ExpectedType, int32 Count, int32 Index, IPCGAttributeAccessorKeys& Keys, EPCGAttributeAccessorFlags Flags, PCG::Private::EPCGAttributeAccessorOperation OpType);
+	virtual bool SetRangeVirtual(PCG::Private::FInValues InValues, int32 Count, int32 Index, IPCGAttributeAccessorKeys& Keys, EPCGAttributeAccessorFlags Flags) = 0;
 
 	bool bReadOnly = true;
 	int16 UnderlyingType = (int16)EPCGMetadataTypes::Unknown;
+	FPCGMetadataAttributeDesc UnderlyingDesc{};
 };
-
-// Specialization of IPCGAttributeAccessor::Get and IPCGAttributeAccessor::Set, for all supported types.
-#define IACCESSOR_DECL(T) \
-template <> bool PCG_API IPCGAttributeAccessor::GetRange<T>(TArrayView<T> OutValues, int32 Index, const IPCGAttributeAccessorKeys& Keys, EPCGAttributeAccessorFlags Flags) const; \
-template <> bool PCG_API IPCGAttributeAccessor::SetRange<T>(TArrayView<const T> InValues, int32 Index, IPCGAttributeAccessorKeys& Keys, EPCGAttributeAccessorFlags Flags);
-PCG_FOREACH_SUPPORTEDTYPES(IACCESSOR_DECL);
-#undef IACCESSOR_DECL
diff --git a/Engine/Plugins/PCG/Source/PCG/Public/Metadata/Accessors/IPCGAttributeAccessorTpl.h b/Engine/Plugins/PCG/Source/PCG/Public/Metadata/Accessors/IPCGAttributeAccessorTpl.h
--- a/Engine/Plugins/PCG/Source/PCG/Public/Metadata/Accessors/IPCGAttributeAccessorTpl.h
+++ b/Engine/Plugins/PCG/Source/PCG/Public/Metadata/Accessors/IPCGAttributeAccessorTpl.h
@@ -22,127 +22,35 @@ protected:
 	IPCGAttributeAccessorT(bool bInReadOnly)
 		: IPCGAttributeAccessor(bInReadOnly, /*UnderlyingType=*/ PCG::Private::MetadataTypes<typename Derived::Type>::Id)
 	{}
-
-	// Override all virtual calls to call our templated method.
-#define IACCESSOR_DECL(T) \
-	virtual bool GetRange##T(TArrayView<T>& OutValues, int32 Index, const IPCGAttributeAccessorKeys& Keys, EPCGAttributeAccessorFlags Flags) const override { return InternalGetRange(OutValues, Index, Keys, Flags); } \
-	virtual bool SetRange##T(TArrayView<const T>& InValues, int32 Index, IPCGAttributeAccessorKeys& Keys, EPCGAttributeAccessorFlags Flags) override { return InternalSetRange(InValues, Index, Keys, Flags); }
-	PCG_FOREACH_SUPPORTEDTYPES(IACCESSOR_DECL);
-#undef IACCESSOR_DECL
-
-private:
-	template<typename U>
-	bool InternalGetRange(TArrayView<U>& OutValues, int32 Index, const IPCGAttributeAccessorKeys& Keys, EPCGAttributeAccessorFlags Flags) const
+	
+	// Default implementation
+	virtual bool GetRangeVirtual(PCG::Private::FOutValues OutValues, int32 Count, int32 Index, const IPCGAttributeAccessorKeys& Keys, EPCGAttributeAccessorFlags Flags) const override
 	{
-		//TRACE_CPUPROFILER_EVENT_SCOPE(IPCGAttributeAccessorT::InternalGetRange);
-
 		using T = typename Derived::Type;
 		const Derived* This = static_cast<const Derived*>(this);
-
-		if constexpr (std::is_same_v<T, U>)
+		
+		if (PCG::Private::FOutValuesByValue* Values = OutValues.TryGet<PCG::Private::FOutValuesByValue>())
 		{
-			return This->GetRangeImpl(OutValues, Index, Keys);
+			return This->GetRangeImpl(MakeArrayView(static_cast<T*>(Values->OutValues), Count), Index, Keys);
 		}
 		else
 		{
-			// Special case'd because FSoftObjectPath currently has a deprecated constructor from FName which generates compile warnings.
-			constexpr bool bNameToSoftObjectPath = std::is_same_v<U, FSoftObjectPath> && std::is_same_v<T, FName>;
-
-			if constexpr (PCG::Private::IsBroadcastable<T, U>())
-			{
-				// Special case - always allow broadcasting from soft object path to string, so that legacy code that grabbed soft path attributes
-				// as FStrings will still work.
-				constexpr bool bSoftReferenceToString = (std::is_same_v<T, FSoftObjectPath> || std::is_same_v<T, FSoftClassPath>) && std::is_same_v<U, FString>;
-				if (!!(Flags & EPCGAttributeAccessorFlags::AllowBroadcast) || bSoftReferenceToString)
-				{
-					TArray<T, TInlineAllocator<4>> InValues;
-					InValues.SetNum(OutValues.Num());
-					if (This->GetRangeImpl(InValues, Index, Keys))
-					{
-						for (int32 i = 0; i < OutValues.Num(); ++i)
-						{
-							PCG::Private::GetValueWithBroadcast<T, U>(InValues[i], OutValues[i]);
-						}
-
-						return true;
-					}
-				}
-			}
-			else if constexpr (std::is_constructible_v<U, T> && !bNameToSoftObjectPath)
-			{
-				if (!!(Flags & EPCGAttributeAccessorFlags::AllowConstructible))
-				{
-					TArray<T, TInlineAllocator<4>> InValues;
-					InValues.SetNum(OutValues.Num());
-					if (This->GetRangeImpl(InValues, Index, Keys))
-					{
-						for (int32 i = 0; i < OutValues.Num(); ++i)
-						{
-							OutValues[i] = U(InValues[i]);
-						}
-						return true;
-					}
-				}
-			}
-			// If U is different of T, and we either can't or allow broadcast or construct from, we will fail, and return false.
 			return false;
 		}
 	}
-
-	template<typename U>
-	bool InternalSetRange(TArrayView<const U>& InValues, int32 Index, IPCGAttributeAccessorKeys& Keys, EPCGAttributeAccessorFlags Flags)
+	
+	virtual bool SetRangeVirtual(PCG::Private::FInValues InValues, int32 Count, int32 Index, IPCGAttributeAccessorKeys& Keys, EPCGAttributeAccessorFlags Flags) override
 	{
-		//TRACE_CPUPROFILER_EVENT_SCOPE(IPCGAttributeAccessorT::InternalSetRange);
-
-		// Can't set if read only, and discard any set that goes beyond the number of keys.
-		if (bReadOnly || Index >= Keys.GetNum())
-		{
-			return false;
-		}
-
 		using T = typename Derived::Type;
 		Derived* This = static_cast<Derived*>(this);
-
-		if constexpr (std::is_same_v<T, U>)
+		
+		if (PCG::Private::FInValuesByValue* Values = InValues.TryGet<PCG::Private::FInValuesByValue>())
 		{
-			return This->SetRangeImpl(InValues, Index, Keys, Flags);
+			return This->SetRangeImpl(MakeArrayView(static_cast<const T*>(Values->InValues), Count), Index, Keys, Flags);
 		}
 		else
 		{
-			// Special case'd because FSoftObjectPath currently has a deprecated constructor from FName which generates compile warnings.
-			constexpr bool bNameToSoftObjectPath = std::is_same_v<T, FSoftObjectPath> && std::is_same_v<U, FName>;
-
-			if constexpr (PCG::Private::IsBroadcastable<U, T>())
-			{
-				if (!!(Flags & EPCGAttributeAccessorFlags::AllowBroadcast))
-				{
-					TArray<T> OutValues;
-					OutValues.SetNum(InValues.Num());
-					for (int32 i = 0; i < OutValues.Num(); ++i)
-					{
-						PCG::Private::GetValueWithBroadcast<U, T>(InValues[i], OutValues[i]);
-					}
-
-					return This->SetRangeImpl(OutValues, Index, Keys, Flags);
-				}
-			}
-			else if constexpr (std::is_constructible_v<T, U> && !bNameToSoftObjectPath)
-			{
-				if (!!(Flags & EPCGAttributeAccessorFlags::AllowConstructible))
-				{
-					TArray<T> OutValues;
-					OutValues.SetNum(InValues.Num());
-					for (int32 i = 0; i < OutValues.Num(); ++i)
-					{
-						OutValues[i] = T(InValues[i]);
-					}
-
-					return This->SetRangeImpl(OutValues, Index, Keys, Flags);
-				}
-			}
-			// If U is different of T, and we either can't or allow broadcast or construct from, we will fail, and return false.
 			return false;
 		}
-
 	}
 };
diff --git a/Engine/Plugins/PCG/Source/PCG/Public/Metadata/PCGMetadataAttributeGeneric.h b/Engine/Plugins/PCG/Source/PCG/Public/Metadata/PCGMetadataAttributeGeneric.h
--- a/Engine/Plugins/PCG/Source/PCG/Public/Metadata/PCGMetadataAttributeGeneric.h
+++ b/Engine/Plugins/PCG/Source/PCG/Public/Metadata/PCGMetadataAttributeGeneric.h
@@ -6,6 +6,7 @@
 #include "Metadata/PCGMetadataCommon.h"
 #include "PCGModule.h"
 
+#include "Helpers/PCGConcepts.h"
 #include "Helpers/PCGMetadataHelpers.h"
 #include "Metadata/PCGMetadataAttribute.h"
 #include "Metadata/PCGMetadataAttributeTraits.h"
@@ -20,24 +21,6 @@ class FPCGMetadataAttributeGeneric;
 namespace PCG::Private
 {
 	void CopyArray(const FArrayProperty* InProperty, void* DestPtr, const void* SrcPtr, int32 Count);
-	
-	template <typename T>
-	struct TIsStaticArray : std::false_type {};
-	
-	template <typename T, uint32 ...Extra>
-    struct TIsStaticArray<TStaticArray<T, Extra...>> : std::true_type {};
-	
-	template <typename Container>
-	concept CIsArrayLike = 
-		TIsTArray_V<Container> 
-		|| TIsTArrayView_V<Container>
-		|| TIsStaticArray<Container>::value;
-
-	template <typename LhsType, typename RhsType>
-	concept CIsComparable = requires(const LhsType& Lhs, const RhsType& Rhs)
-	{
-		{ Lhs == Rhs } -> UE::CSameAs<bool>;
-	};
 
 	/** Magic number for AddValues to know if we need to write to the DefaultValue or not. */
 	constexpr int32 ExistingIndexForDefaultValue = -2;
@@ -111,6 +94,7 @@ public:
 	PCG_API bool IsMap() const;
 
 	const FPCGMetadataAttributeDesc& GetAttributeDesc() const { return CachedDesc; }
+	bool DoesCompressData() const { return UnderlyingProperty && UnderlyingProperty->bCompressData; }
 
 	template <typename T>
 	bool IsOfType() const
@@ -155,7 +139,7 @@ public:
 			return;
 		}
 		
-		GetValues_Internal(ValueKeys, FOutValues{TInPlaceType<FOutValuesByValue>{}, OutValues.GetData()});
+		GetValues_Internal(ValueKeys, PCG::Private::FOutValues{TInPlaceType<PCG::Private::FOutValuesByValue>{}, OutValues.GetData()});
 	}
 
 	template <typename T>
@@ -181,7 +165,7 @@ public:
 		// We need to type erase the pointers, but to be fully C++ compliant about aliasing rules we can't reinterpret_cast. So make a copy of the pointers.
 		TArray<const void*, TInlineAllocator<256>> Temp;
 		Temp.SetNumUninitialized(OutValues.Num());
-		GetValues_Internal(ValueKeys, FOutValues{ TInPlaceType<FOutValuesByPtr>{}, Temp });
+		GetValues_Internal(ValueKeys, PCG::Private::FOutValues{ TInPlaceType<PCG::Private::FOutValuesByPtr>{}, Temp });
 		for (int32 i = 0; i < OutValues.Num(); ++i)
 		{
 			OutValues[i] = static_cast<const T*>(Temp[i]);
@@ -210,7 +194,7 @@ public:
 		
 		TArray<TTuple<const void*, int32>, TInlineAllocator<256>> Temp;
 		Temp.SetNumUninitialized(OutValues.Num());
-		GetValues_Internal(ValueKeys, FOutValues{TInPlaceType<FOutValuesAsArray>{}, Temp});
+		GetValues_Internal(ValueKeys, PCG::Private::FOutValues{TInPlaceType<PCG::Private::FOutValuesAsArray>{}, Temp});
 		for (int32 i = 0; i < OutValues.Num(); ++i)
 		{
 			OutValues[i] = MakeConstArrayView<T>(static_cast<const T*>(Temp[i].Get<0>()), Temp[i].Get<1>());
@@ -244,7 +228,7 @@ public:
 		{
 			Temp.Add(&Value.Helper);
 		}
-		GetValues_Internal(ValueKeys, FOutValues{TInPlaceType<FOutValuesAsSet>{}, Temp});
+		GetValues_Internal(ValueKeys, PCG::Private::FOutValues{TInPlaceType<PCG::Private::FOutValuesAsSet>{}, Temp});
 	}
 	
 	// To support maps
@@ -274,7 +258,7 @@ public:
 		{
 			Temp.Add(&Value.Helper);
 		}
-		GetValues_Internal(ValueKeys, FOutValues{TInPlaceType<FOutValuesAsMap>{}, Temp});
+		GetValues_Internal(ValueKeys, PCG::Private::FOutValues{TInPlaceType<PCG::Private::FOutValuesAsMap>{}, Temp});
 	}
 	
 	////////////// End GetValues Section
@@ -394,7 +378,7 @@ public:
 protected:
 	/** Attributes can't be created individually, they need to be created through a Metadata Domain. */
 	friend FPCGMetadataDomain;
-
+	
 	// New attribute with no parent, name is contained in the InAttributeDesc.
 	PCG_API FPCGMetadataAttributeGeneric(const FPCGMetadataAttributeDesc& InAttributeDesc, FPCGMetadataDomain* InMetadata, bool bInAllowsInterpolation);
 
@@ -421,34 +405,7 @@ protected:
 	 * GetValues_Internal is working on a Variant that will type erase the underlying output container
 	 * in order to reduce template code size.
 	 */
-	struct FOutValuesByValue
-	{
-		void* OutValues;
-	};
-	
-	struct FOutValuesByPtr
-	{
-		TArrayView<const void*> OutValues;
-	};
-	
-	struct FOutValuesAsArray
-	{
-		TArrayView<TTuple<const void*, int32>> OutValues;
-	};
-	
-	struct FOutValuesAsSet
-	{
-		TArrayView<FScriptSetHelper*> OutValues;
-	};
-	
-	struct FOutValuesAsMap
-	{
-		TArrayView<FScriptMapHelper*> OutValues;
-	};
-	
-	using FOutValues = TVariant<FOutValuesByValue, FOutValuesByPtr, FOutValuesAsArray, FOutValuesAsSet, FOutValuesAsMap>;
-	
-	void GetValues_Internal(const TArrayView<const PCGMetadataValueKey> ValueKeys, FOutValues OutValues) const
+	void GetValues_Internal(const TArrayView<const PCGMetadataValueKey> ValueKeys, PCG::Private::FOutValues OutValues) const
 	{
 		// Bitset with all unretrieved values. If we have any unretrieved value, we will ask the parent for those.
 		TBitArray<> UnretrievedValues(true, ValueKeys.Num());
@@ -466,34 +423,6 @@ protected:
 	 * is necessary or not. This is a very niche feature, and it is expected that everything is already allocated
 	 * or it will crash. There is also an option for lockless, if we expect to be able to write without any lock.
 	 */
-	struct FInValuesByValue
-	{
-		const void* InValues;
-		int32 Count;
-	};
-
-	struct FInValuesByValueSubset
-	{
-		const void* InValues;
-		TArray<int32> Indices;
-	};
-
-	struct FInValuesAsArray
-	{
-		TArrayView<TTuple<const void*, int32>> InValues;
-	};
-
-	struct FInValuesAsSet
-	{
-		TArrayView<TArray<const void*>> InValues;
-	};
-
-	struct FInValuesAsMap
-	{
-		TArrayView<TArray<TPair<const void*, const void*>>> InValues;
-	};
-
-	using FInValues = TVariant<FInValuesByValue, FInValuesAsArray, FInValuesAsSet, FInValuesAsMap, FInValuesByValueSubset>;
 
 	/**
 	 * Basic version that will accept all single values (as an array view).
@@ -515,7 +444,7 @@ protected:
 
 		// For compressed data, we need to analyze which values actually need to be added. If existing index is not none, ignore it.
 		// Need the concept for Equals, since `AddUnique` require the operator '==` to be defined.
-		if constexpr (PCG::Private::CIsComparable<T, T>)
+		if constexpr (PCG::Concepts::CIsComparable<T, T>)
 		{
 			if (UnderlyingProperty->bCompressData && ExistingIndex == INDEX_NONE)
 			{
@@ -532,17 +461,17 @@ protected:
 					FoundValueKeys.Emplace(UniqueValues.AddUnique(InValues[ValueIndex]));
 				}
 
-				return AddCompressedValues(FInValuesByValue{UniqueValues.GetData(), UniqueValues.Num()}, FoundValueKeys);
+				return AddCompressedValues(PCG::Private::FInValuesByValue{UniqueValues.GetData(), UniqueValues.Num()}, FoundValueKeys);
 			}
 		}
 
-		return AddValues_Internal(FInValues{TInPlaceType<FInValuesByValue>{}, InValues.GetData(), InValues.Num()}, InValues.Num(), ExistingIndex, bLockless);
+		return AddValues_Internal(PCG::Private::FInValues{TInPlaceType<PCG::Private::FInValuesByValue>{}, InValues.GetData(), InValues.Num()}, InValues.Num(), ExistingIndex, bLockless);
 	}
 	
 	/**
 	 * Version that accept anything that is array like, like `TArray<T>`, `TArrayView<T>`, `TStaticArray<T, 5>`, `TArray<T, TInlineAllocator<5>>`, etc...
 	 */
-	template <typename ArrayContainer> requires (PCG::Private::CIsArrayLike<ArrayContainer>)
+	template <typename ArrayContainer> requires (PCG::Concepts::CIsArrayLike<ArrayContainer>)
 	TArray<PCGMetadataValueKey> AddValues(TConstArrayView<ArrayContainer> InValues, int32 ExistingIndex = INDEX_NONE, bool bLockless = false)
 	{
 		if (InValues.IsEmpty())
@@ -561,7 +490,7 @@ protected:
 		TArray<TTuple<const void*, int32>, TInlineAllocator<256>> InValuesTuple;
 		InValuesTuple.Reserve(InValues.Num());
 		Algo::Transform(InValues, InValuesTuple, [](const auto& Item) { return TTuple<const void*, int32>{Item.GetData(), Item.Num()}; });
-		return AddValues_Internal(FInValues{TInPlaceType<FInValuesAsArray>{}, InValuesTuple}, InValues.Num(), ExistingIndex, bLockless);
+		return AddValues_Internal(PCG::Private::FInValues{TInPlaceType<PCG::Private::FInValuesAsArray>{}, InValuesTuple}, InValues.Num(), ExistingIndex, bLockless);
 	}
 
 	/**
@@ -591,7 +520,7 @@ protected:
 				Algo::Transform(Item, SetValues, [](const T& Value) -> const T* { return &Value;});
 				return SetValues;
 			});
-		return AddValues_Internal(FInValues{TInPlaceType<FInValuesAsSet>{}, InSetValues}, InValues.Num(), ExistingIndex, bLockless);
+		return AddValues_Internal(PCG::Private::FInValues{TInPlaceType<PCG::Private::FInValuesAsSet>{}, InSetValues}, InValues.Num(), ExistingIndex, bLockless);
 	}
 	
 	/**
@@ -621,24 +550,24 @@ protected:
 				Algo::Transform(Item, MapValues, [](const auto& It) -> TPair<const void*, const void*> { return {&It.Key, &It.Value};});
 				return MapValues;
 			});
-		return AddValues_Internal(FInValues{TInPlaceType<FInValuesAsMap>{}, InMapValues}, InValues.Num(), ExistingIndex, bLockless);
+		return AddValues_Internal(PCG::Private::FInValues{TInPlaceType<PCG::Private::FInValuesAsMap>{}, InMapValues}, InValues.Num(), ExistingIndex, bLockless);
 	}
 	
 	// Internal implementation that will do the dispatch using the typed erased variant.
-	PCG_API TArray<PCGMetadataValueKey> AddValues_Internal(FInValues InValues, int32 Count, int32 ExistingIndex = INDEX_NONE, bool bLockless = false);
+	PCG_API TArray<PCGMetadataValueKey> AddValues_Internal(const PCG::Private::FInValues& InValues, int32 Count, int32 ExistingIndex = INDEX_NONE, bool bLockless = false);
 
 	// Internal implementation for compressed data
-	PCG_API TArray<PCGMetadataValueKey> AddCompressedValues(FInValuesByValue InUniqueValues, TArray<PCGMetadataValueKey>& FoundValueKeys);
+	PCG_API TArray<PCGMetadataValueKey> AddCompressedValues(PCG::Private::FInValuesByValue InUniqueValues, TArray<PCGMetadataValueKey>& FoundValueKeys);
 
 	/** 
 	 * Find all the value keys associated to InValues for compressed data. Return true if all the values were found.
 	 */
-	PCG_API bool FindValues(TArray<PCGMetadataValueKey>& OutValueKeys, const FInValuesByValue& InValues) const;
+	PCG_API bool FindValues(TArray<PCGMetadataValueKey>& OutValueKeys, const PCG::Private::FInValuesByValue& InValues) const;
 
 	/** 
 	 * Find all the value keys associated to InValues for compressed data. Return true if all the values were found. Preallocated version.
 	 */
-	PCG_API bool FindValues(TArrayView<PCGMetadataValueKey> OutValueKeys, const FInValuesByValue& InValues) const;
+	PCG_API bool FindValues(TArrayView<PCGMetadataValueKey> OutValueKeys, const PCG::Private::FInValuesByValue& InValues) const;
 
 	template <typename T>
 	void InitForDeprecation(TConstArrayView<T> InValues, const T& InDefaultValue)
@@ -665,8 +594,8 @@ protected:
 private:
 	void Init();
 	
-	PCG_API void GetValues_Internal(const TArrayView<const PCGMetadataValueKey> ValueKeys, FOutValues OutValues, TBitArray<>& UnretrievedValues) const;
-	void FindValuesInternal(const FInValuesByValue& InValues, TArrayView<PCGMetadataValueKey> ValueKeys, int& ValueKeysSet, bool bIsRoot) const;
+	PCG_API void GetValues_Internal(const TArrayView<const PCGMetadataValueKey> ValueKeys, PCG::Private::FOutValues OutValues, TBitArray<>& UnretrievedValues) const;
+	void FindValuesInternal(const PCG::Private::FInValuesByValue& InValues, TArrayView<PCGMetadataValueKey> ValueKeys, int& ValueKeysSet, bool bIsRoot) const;
 
 public:
 	PCG_API virtual void Flatten() override;
diff --git a/Engine/Plugins/PCG/Source/PCG/Public/Metadata/PCGMetadataAttributeTraits.h b/Engine/Plugins/PCG/Source/PCG/Public/Metadata/PCGMetadataAttributeTraits.h
--- a/Engine/Plugins/PCG/Source/PCG/Public/Metadata/PCGMetadataAttributeTraits.h
+++ b/Engine/Plugins/PCG/Source/PCG/Public/Metadata/PCGMetadataAttributeTraits.h
@@ -2,6 +2,7 @@
 
 #pragma once
 
+#include "Helpers/PCGConcepts.h"
 #include "Metadata/PCGMetadataCommon.h"
 
 #include "Math/Transform.h" // IWYU pragma: keep
@@ -36,24 +37,52 @@ namespace PCG
 		struct TMapTypeExtractor
 		{
 		};
-		
+
 		template <typename InKeyType, typename InValueType>
 		struct TMapTypeExtractor<TMap<InKeyType, InValueType>>
 		{
 			using KeyType = InKeyType;
 			using ValueType = InValueType;
 		};
-		
+
 		// Utility struct to know when a type is a basic type, specialized in the macro.
 		template <typename T> struct TIsBasicType { enum { Value = false }; };
-		
-		// Concept to catch all "basic types" that are also structs.
+
+		// To transform a range of values from one type to another.
+		using FTransformFunc = void(*)(void* /*DestPtr*/, const void* /*SrcPtr*/, int32 /*Count*/);
+
+		// To construct a range of values.
+		using FConstructMemoryFunc = void(*)(void* /*DestPtr*/, int32 /*Count*/);
+
+		// To destruct a range of values.
+		using FDestroyMemoryFunc = void(*)(void* /*DestPtr*/, int32 /*Count*/);
+
 		template <typename T>
-		concept CIsBasicStruct = requires (T)
+		void ConstructMemory(void* DestPtr, int32 Count)
 		{
-			TBaseStructure<T>::Get();
-		};
-		
+			TArrayView<T> Dest(static_cast<T*>(DestPtr), Count);
+			for (T& Value : Dest)
+			{
+				new (&Value) T();
+			}
+		}
+
+		template <typename T>
+		void DestructMemory(void* DestPtr, int32 Count)
+		{
+			DestructItems(static_cast<T*>(DestPtr), Count);
+		}
+
+		// Forward declare
+		template <typename InType, typename OutType>
+		inline bool GetValueWithBroadcast(const InType& InValue, OutType& OutValue);
+
+		template <typename InType, typename OutType>
+		inline void GetValuesWithBroadcast_Internal(void* DestPtr, const void* SrcPtr, int32 Count);
+
+		template <typename InType, typename OutType>
+		inline void GetValuesWithConstructible_Internal(void* DestPtr, const void* SrcPtr, int32 Count);
+
 		// Old way of converting types to the type enum, will be progressively replaced by MakeAttributeDesc/IsSameType.
 		template <typename T>
 		struct MetadataTypes
@@ -62,16 +91,18 @@ namespace PCG
 		};
 
 		/**
-		 * Return the attribute descriptor associated with the type T, with the Attribute name provided.
+		 * Return the default attribute descriptor associated with the type T, without any attribute name.
 		 * This version does not compile and is a pitfall for all unsupported types.
 		 */
 		template <typename T>
-		FPCGMetadataAttributeDesc MakeAttributeDesc(FName AttributeName)
+		const FPCGMetadataAttributeDesc& GetDefaultAttributeDesc()
 		{
-			static_assert(!std::is_same_v<T,T>, "This type is not supported for metadata."); 
-			return FPCGMetadataAttributeDesc{};
+			static const FPCGMetadataAttributeDesc EmptyDesc{};
+
+			static_assert(!std::is_same_v<T,T>, "This type is not supported for metadata.");
+			return EmptyDesc;
 		}
-		
+
 		/**
 		 * Return true if the provided Desc is matching with type T.
 		 * This version does not compile and is a pitfall for all unsupported types.
@@ -79,10 +110,10 @@ namespace PCG
 		template <typename T>
 		bool IsEquivalentDesc(const FPCGMetadataAttributeDesc& Desc)
 		{
-			static_assert(!std::is_same_v<T,T>, "This type is not supported for metadata."); 
+			static_assert(!std::is_same_v<T,T>, "This type is not supported for metadata.");
 			return false;
 		}
-		
+
 		/**
 		 * Return true if the provided type and type object is matching with type T.
 		 * This version does not compile for non-basic types and is a pitfall for all unsupported types.
@@ -92,10 +123,10 @@ namespace PCG
 		{
 			if constexpr (!TIsBasicType<T>::Value)
 			{
-				static_assert(!std::is_same_v<T,T>, "This type is not supported for metadata."); 
+				static_assert(!std::is_same_v<T,T>, "This type is not supported for metadata.");
 				return false;
 			}
-			else if constexpr (CIsBasicStruct<T>) 
+			else if constexpr (PCG::Concepts::CIsBasicStruct<T>)
 			{
 				// Basic structures can be represented with their own EPCGMetadataTypes (like Vector) or as a struct.
 				return static_cast<uint16>(InType) == MetadataTypes<T>::Id || (InType == EPCGMetadataTypes::Struct && InTypeObject == TBaseStructure<T>::Get());
@@ -111,7 +142,7 @@ namespace PCG
 		 */
 #define PCGMetadataGenerateDataTypes(Type, TypeEnum) \
 	template<> struct MetadataTypes<Type>{ enum { Id = static_cast<uint16>(EPCGMetadataTypes::TypeEnum)}; }; \
-	template<> inline FPCGMetadataAttributeDesc MakeAttributeDesc<Type>(FName AttributeName) { return FPCGMetadataAttributeDesc{.Name=AttributeName, .ValueType=EPCGMetadataTypes::TypeEnum}; } \
+	template<> inline const FPCGMetadataAttributeDesc& GetDefaultAttributeDesc<Type>() { static const FPCGMetadataAttributeDesc TypedDesc{.ValueType=EPCGMetadataTypes::TypeEnum}; return TypedDesc; } \
 	template<> struct TIsBasicType<Type> { enum { Value = true }; }; \
 	template<> inline bool IsEquivalentDesc<Type>(const FPCGMetadataAttributeDesc& Desc) { return IsSameType<Type>(Desc.ValueType, Desc.ValueTypeObject) && Desc.ContainerTypes.IsEmpty(); }
 
@@ -136,7 +167,7 @@ namespace PCG
 #undef PCGMetadataGenerateDataTypes
 
 		/**
-		 * Specialization section of MakeAttributeDesc
+		 * Specialization section of GetDefaultAttributeDesc
 		 * Supported generic types:
 		 *   - UStruct
 		 *   - UEnum
@@ -145,114 +176,146 @@ namespace PCG
 		 *   - TMap<T, U> if T and U are supported types. Does not support nested containers for now.
 		 */
 		template<typename T> requires (TModels_V<CStaticStructProvider, T> && !std::is_base_of_v<FSoftObjectPath, T>)
-		FPCGMetadataAttributeDesc MakeAttributeDesc(FName AttributeName)
+		const FPCGMetadataAttributeDesc& GetDefaultAttributeDesc()
 		{
-			return FPCGMetadataAttributeDesc{
-				.Name = AttributeName,
+			static const FPCGMetadataAttributeDesc TypedDesc{
 				.ValueType = EPCGMetadataTypes::Struct,
 				.ValueTypeObject = T::StaticStruct()
 			};
+
+			return TypedDesc;
 		}
-		
+
 		template<typename T> requires (TIsUEnumClass<T>::Value == true)
-		FPCGMetadataAttributeDesc MakeAttributeDesc(FName AttributeName)
+		const FPCGMetadataAttributeDesc& GetDefaultAttributeDesc()
 		{
-			return FPCGMetadataAttributeDesc{
-				.Name = AttributeName,
+			static const FPCGMetadataAttributeDesc TypedDesc{
 				.ValueType = EPCGMetadataTypes::Enum,
 				.ValueTypeObject = StaticEnum<T>()
 			};
+
+			return TypedDesc;
 		}
-		
+
 		template <typename Container> requires TIsTArray_V<Container>
-		FPCGMetadataAttributeDesc MakeAttributeDesc(FName AttributeName)
+		const FPCGMetadataAttributeDesc& GetDefaultAttributeDesc()
 		{
 			using T = typename Container::ElementType;
-			FPCGMetadataAttributeDesc Result = MakeAttributeDesc<T>(AttributeName);
-			Result.ContainerTypes.Add(EPCGMetadataAttributeContainerTypes::Array);
-			return Result;
+
+			static const FPCGMetadataAttributeDesc TypedDesc = []()
+			{
+				FPCGMetadataAttributeDesc Result = GetDefaultAttributeDesc<T>();
+				Result.ContainerTypes.Add(EPCGMetadataAttributeContainerTypes::Array);
+				return Result;
+			}();
+
+			return TypedDesc;
 		}
-		
+
 		template <typename Container> requires (TIsTSet<Container>::Value == true)
-		FPCGMetadataAttributeDesc MakeAttributeDesc(FName AttributeName)
+		const FPCGMetadataAttributeDesc& GetDefaultAttributeDesc()
 		{
 			using T = typename Container::ElementType;
-			FPCGMetadataAttributeDesc Result = MakeAttributeDesc<T>(AttributeName);
-			Result.ContainerTypes.Add(EPCGMetadataAttributeContainerTypes::Set);
-			return Result;
+
+			static const FPCGMetadataAttributeDesc TypedDesc = []()
+			{
+				FPCGMetadataAttributeDesc Result = GetDefaultAttributeDesc<T>();
+				Result.ContainerTypes.Add(EPCGMetadataAttributeContainerTypes::Set);
+				return Result;
+			}();
+
+			return TypedDesc;
 		}
-		
+
 		template<typename Container> requires (TIsTMap<Container>::Value == true)
-		FPCGMetadataAttributeDesc MakeAttributeDesc(FName AttributeName)
+		const FPCGMetadataAttributeDesc& GetDefaultAttributeDesc()
 		{
 			using KeyType = typename TMapTypeExtractor<Container>::KeyType;
 			using ValueType = typename TMapTypeExtractor<Container>::ValueType;
-			
-			FPCGMetadataAttributeDesc Result = MakeAttributeDesc<ValueType>(AttributeName);
-			FPCGMetadataAttributeDesc Key = MakeAttributeDesc<KeyType>(AttributeName);
-			
-			Result.ContainerTypes.Add(EPCGMetadataAttributeContainerTypes::Map);
-			Result.KeyType = Key.ValueType;
-			Result.KeyTypeObject = Key.ValueTypeObject;
+
+			static const FPCGMetadataAttributeDesc TypedDesc = []()
+			{
+				FPCGMetadataAttributeDesc Result = GetDefaultAttributeDesc<ValueType>();
+				const FPCGMetadataAttributeDesc& Key = GetDefaultAttributeDesc<KeyType>();
+
+				Result.ContainerTypes.Add(EPCGMetadataAttributeContainerTypes::Map);
+				Result.KeyType = Key.ValueType;
+				Result.KeyTypeObject = Key.ValueTypeObject;
+
+				return Result;
+			}();
+
+			return TypedDesc;
+		}
+
+		//////////////////////////////
+
+		/**
+		 * Return the attribute descriptor associated with the type T, with the Attribute name provided.
+		 */
+		template <typename T>
+		FPCGMetadataAttributeDesc MakeAttributeDesc(FName AttributeName)
+		{
+			FPCGMetadataAttributeDesc Result = GetDefaultAttributeDesc<T>();
+			Result.Name = AttributeName;
 			return Result;
 		}
-		
+
 		//////////////////////////////
-		
+
 		template<typename T> requires (TModels_V<CStaticStructProvider, T>  && !std::is_base_of_v<FSoftObjectPath, T>)
 		bool IsSameType(const EPCGMetadataTypes InType, const UObject* InTypeObject)
 		{
 			return InType == EPCGMetadataTypes::Struct &&
 				InTypeObject == T::StaticStruct();
 		}
-		
+
 		template<typename T> requires (TIsUEnumClass<T>::Value == true)
 		bool IsSameType(const EPCGMetadataTypes InType, const UObject* InTypeObject)
 		{
 			return InType == EPCGMetadataTypes::Enum &&
 				InTypeObject == StaticEnum<T>();
 		}
-		
+
 		//////////////////////////////
 
 		template<typename T> requires ((TModels_V<CStaticStructProvider, T>  && !std::is_base_of_v<FSoftObjectPath, T>) || TIsUEnumClass<T>::Value == true)
 		bool IsEquivalentDesc(const FPCGMetadataAttributeDesc& Desc)
 		{
-			return Desc.ContainerTypes.IsEmpty() 
+			return Desc.ContainerTypes.IsEmpty()
 				&& IsSameType<T>(Desc.ValueType, Desc.ValueTypeObject);
 		}
-		
+
 		template <typename Container> requires TIsTArray_V<Container>
 		bool IsEquivalentDesc(const FPCGMetadataAttributeDesc& Desc)
 		{
 			using T = typename Container::ElementType;
-			return !Desc.ContainerTypes.IsEmpty() 
-				&& Desc.ContainerTypes[0] == EPCGMetadataAttributeContainerTypes::Array 
+			return !Desc.ContainerTypes.IsEmpty()
+				&& Desc.ContainerTypes[0] == EPCGMetadataAttributeContainerTypes::Array
 				&& IsSameType<T>(Desc.ValueType, Desc.ValueTypeObject);
 		}
-		
+
 		template <typename Container> requires (TIsTSet<Container>::Value == true)
 		bool IsEquivalentDesc(const FPCGMetadataAttributeDesc& Desc)
 		{
 			using T = typename Container::ElementType;
-			return !Desc.ContainerTypes.IsEmpty() 
-				&& Desc.ContainerTypes[0] == EPCGMetadataAttributeContainerTypes::Set 
+			return !Desc.ContainerTypes.IsEmpty()
+				&& Desc.ContainerTypes[0] == EPCGMetadataAttributeContainerTypes::Set
 				&& IsSameType<T>(Desc.ValueType, Desc.ValueTypeObject);
-
 		}
-		
+
 		template<typename Container> requires (TIsTMap<Container>::Value == true)
 		bool IsEquivalentDesc(const FPCGMetadataAttributeDesc& Desc)
 		{
 			using KeyType = typename TMapTypeExtractor<Container>::KeyType;
 			using ValueType = typename TMapTypeExtractor<Container>::ValueType;
-			
-			return !Desc.ContainerTypes.IsEmpty() 
-				&& Desc.ContainerTypes[0] == EPCGMetadataAttributeContainerTypes::Map 
+
+			return !Desc.ContainerTypes.IsEmpty()
+				&& Desc.ContainerTypes[0] == EPCGMetadataAttributeContainerTypes::Map
 				&& IsSameType<KeyType>(Desc.KeyType, Desc.KeyTypeObject)
 				&& IsSameType<ValueType>(Desc.ValueType, Desc.ValueTypeObject);
 		}
-		
+
 		////////////////////////
 
 		constexpr inline bool IsPCGType(uint16 TypeId)
@@ -330,7 +393,9 @@ namespace PCG
 		// First index is the original type, second index is the wanted type. Returns true if we can broadcast first type into second type.
 		struct FBroadcastableTypes
 		{
-			constexpr FBroadcastableTypes() : Values{ {false} }
+			constexpr FBroadcastableTypes()
+				: Values{ {false} }
+				, TransformFuncs { {nullptr} }
 			{
 				// All types are broadcastable to themselves.
 				for (uint8 i = 0; i < (uint8)EPCGMetadataTypes::Count; ++i)
@@ -338,59 +403,62 @@ namespace PCG
 					Values[i][i] = true;
 				}
 
-#define PCGMetadataBroadcastable(FirstType, SecondType) Values[(uint8)EPCGMetadataTypes::FirstType][(uint8)EPCGMetadataTypes::SecondType] = true
+#define PCGMetadataBroadcastable(FirstType, SecondType) \
+	Values[MetadataTypes<FirstType>::Id][MetadataTypes<SecondType>::Id] = true;\
+	TransformFuncs[MetadataTypes<FirstType>::Id][MetadataTypes<SecondType>::Id] = &GetValuesWithBroadcast_Internal<FirstType, SecondType>
 
 				// Finally set all cases where it is broadcastable
-				PCGMetadataBroadcastable(Boolean, Integer32);
-				PCGMetadataBroadcastable(Boolean, Integer64);
-				PCGMetadataBroadcastable(Boolean, Double);
-
-				PCGMetadataBroadcastable(Float, Double);
-				PCGMetadataBroadcastable(Float, Vector2);
-				PCGMetadataBroadcastable(Float, Vector);
-				PCGMetadataBroadcastable(Float, Vector4);
-
-				PCGMetadataBroadcastable(Double, Vector2);
-				PCGMetadataBroadcastable(Double, Vector);
-				PCGMetadataBroadcastable(Double, Vector4);
-
-				PCGMetadataBroadcastable(Integer32, Float);
-				PCGMetadataBroadcastable(Integer32, Double);
-				PCGMetadataBroadcastable(Integer32, Integer64);
-				PCGMetadataBroadcastable(Integer32, Vector2);
-				PCGMetadataBroadcastable(Integer32, Vector);
-				PCGMetadataBroadcastable(Integer32, Vector4);
-
-				PCGMetadataBroadcastable(Integer64, Float);
-				PCGMetadataBroadcastable(Integer64, Double);
-				PCGMetadataBroadcastable(Integer64, Vector2);
-				PCGMetadataBroadcastable(Integer64, Vector);
-				PCGMetadataBroadcastable(Integer64, Vector4);
-
-				PCGMetadataBroadcastable(Vector2, Vector);
-				PCGMetadataBroadcastable(Vector2, Vector4);
-
-				PCGMetadataBroadcastable(Vector, Vector4);
-
-				PCGMetadataBroadcastable(Quaternion, Rotator);
-				PCGMetadataBroadcastable(Rotator, Quaternion);
-
-				PCGMetadataBroadcastable(Float, String);
-				PCGMetadataBroadcastable(Double, String);
-				PCGMetadataBroadcastable(Integer32, String);
-				PCGMetadataBroadcastable(Integer64, String);
-				PCGMetadataBroadcastable(Vector2, String);
-				PCGMetadataBroadcastable(Vector, String);
-				PCGMetadataBroadcastable(Vector4, String);
-				PCGMetadataBroadcastable(Quaternion, String);
-				PCGMetadataBroadcastable(Transform, String);
-				PCGMetadataBroadcastable(Boolean, String);
-				PCGMetadataBroadcastable(Rotator, String);
-				PCGMetadataBroadcastable(Name, String);
-				PCGMetadataBroadcastable(String, Name);
-
-				PCGMetadataBroadcastable(SoftObjectPath, String);
-				PCGMetadataBroadcastable(SoftClassPath, String);
+				PCGMetadataBroadcastable(bool, int32);
+				PCGMetadataBroadcastable(bool, int64);
+				PCGMetadataBroadcastable(bool, float);
+				PCGMetadataBroadcastable(bool, double);
+
+				PCGMetadataBroadcastable(float, double);
+				PCGMetadataBroadcastable(float, FVector2D);
+				PCGMetadataBroadcastable(float, FVector);
+				PCGMetadataBroadcastable(float, FVector4);
+
+				PCGMetadataBroadcastable(double, FVector2D);
+				PCGMetadataBroadcastable(double, FVector);
+				PCGMetadataBroadcastable(double, FVector4);
+
+				PCGMetadataBroadcastable(int32, float);
+				PCGMetadataBroadcastable(int32, double);
+				PCGMetadataBroadcastable(int32, int64);
+				PCGMetadataBroadcastable(int32, FVector2D);
+				PCGMetadataBroadcastable(int32, FVector);
+				PCGMetadataBroadcastable(int32, FVector4);
+
+				PCGMetadataBroadcastable(int64, float);
+				PCGMetadataBroadcastable(int64, double);
+				PCGMetadataBroadcastable(int64, FVector2D);
+				PCGMetadataBroadcastable(int64, FVector);
+				PCGMetadataBroadcastable(int64, FVector4);
+
+				PCGMetadataBroadcastable(FVector2D, FVector);
+				PCGMetadataBroadcastable(FVector2D, FVector4);
+
+				PCGMetadataBroadcastable(FVector, FVector4);
+
+				PCGMetadataBroadcastable(FQuat, FRotator);
+				PCGMetadataBroadcastable(FRotator, FQuat);
+
+				PCGMetadataBroadcastable(float, FString);
+				PCGMetadataBroadcastable(double, FString);
+				PCGMetadataBroadcastable(int32, FString);
+				PCGMetadataBroadcastable(int64, FString);
+				PCGMetadataBroadcastable(FVector2D, FString);
+				PCGMetadataBroadcastable(FVector, FString);
+				PCGMetadataBroadcastable(FVector4, FString);
+				PCGMetadataBroadcastable(FQuat, FString);
+				PCGMetadataBroadcastable(FTransform, FString);
+				PCGMetadataBroadcastable(bool, FString);
+				PCGMetadataBroadcastable(FRotator, FString);
+				PCGMetadataBroadcastable(FName, FString);
+				PCGMetadataBroadcastable(FString, FName);
+
+				PCGMetadataBroadcastable(FSoftObjectPath, FString);
+				PCGMetadataBroadcastable(FSoftClassPath, FString);
 
 #undef PCGMetadataBroadcastable
 			}
@@ -398,6 +466,8 @@ namespace PCG
 			constexpr const bool* operator[](int i) const { return Values[i]; }
 
 			bool Values[(uint8)EPCGMetadataTypes::Count][(uint8)EPCGMetadataTypes::Count];
+
+			FTransformFunc TransformFuncs[(uint8)EPCGMetadataTypes::Count][(uint8)EPCGMetadataTypes::Count];
 		};
 
 		// Wrapper around a standard 2-dimensional CArray that is constexpr, to know if a type is constructible from another.
@@ -416,37 +486,89 @@ namespace PCG
 #undef PCG_ALL_TYPES
 
 			// Unroll AllTypes twice, to have all combinations of types <T, U>
-			template <typename FirstType> constexpr void UnrollSecond(TypeHolder<void> InTypeHolder) {}
+			template <typename OutputType> constexpr void UnrollSecond(TypeHolder<void> InTypeHolder) {}
 
-			template <typename FirstType, typename SecondType, typename... SecondTypes>
-			constexpr void UnrollSecond(TypeHolder<SecondType, SecondTypes...> InTypeHolder)
+			template <typename OutputType, typename InputType, typename... InputTypes>
+			constexpr void UnrollSecond(TypeHolder<InputType, InputTypes...> InTypeHolder)
 			{
-				Values[MetadataTypes<FirstType>::Id][MetadataTypes<SecondType>::Id] = std::is_constructible_v<FirstType, SecondType>;
-				UnrollSecond<FirstType>(TypeHolder<SecondTypes...>{});
+				if constexpr (std::is_constructible_v<OutputType, InputType>)
+				{
+					Values[MetadataTypes<OutputType>::Id][MetadataTypes<InputType>::Id] = true;
+					TransformFuncs[MetadataTypes<OutputType>::Id][MetadataTypes<InputType>::Id] = &GetValuesWithConstructible_Internal<InputType, OutputType>;
+				}
+				UnrollSecond<OutputType>(TypeHolder<InputTypes...>{});
 			}
 
-			template <typename InputType, typename... InputTypes>
-			constexpr void UnrollFirst(TypeHolder<InputType, InputTypes...> InTypeHolder)
+			template <typename OutputType, typename... OutputTypes>
+			constexpr void UnrollFirst(TypeHolder<OutputType, OutputTypes...> InTypeHolder)
 			{
-				UnrollSecond<InputType>(AllTypes{});
-				UnrollFirst(TypeHolder<InputTypes...>{});
+				UnrollSecond<OutputType>(AllTypes{});
+				UnrollFirst(TypeHolder<OutputTypes...>{});
 			}
 
 			constexpr void UnrollFirst(TypeHolder<void> InTypeHolder) {}
 
 		public:
-			constexpr FConstructibleTypes() : Values{ {false} }
+			constexpr FConstructibleTypes()
+				: Values{ {false} }
+				, TransformFuncs{ {nullptr} }
 			{
 				UnrollFirst(AllTypes{});
 			}
 
 			constexpr const bool* operator[](int i) const { return Values[i]; }
 
+			// First type = OutputType, Second type = Input type. Can I construct OutputType from InputType.
 			bool Values[(uint8)EPCGMetadataTypes::Count][(uint8)EPCGMetadataTypes::Count];
+			FTransformFunc TransformFuncs[(uint8)EPCGMetadataTypes::Count][(uint8)EPCGMetadataTypes::Count];
+		};
+
+		// Wrapper around a standard 1-dimensional CArray that is constexpr, to know how to construct/destruct/element size
+		struct FTypesFundamental
+		{
+		private:
+			// Empty struct that will hold all our types.
+			template <typename... T>
+			struct TypeHolder {};
+
+			// Use the macro to have AllTypes that is the Dummy struct templated with all our types.
+			// Add void as a sentinel, to know when to stop unrolling (and because TypeHolder<T,U,> is syntax error)
+#define PCG_ALL_TYPES(T) T,
+			using AllTypes = TypeHolder<PCG_FOREACH_SUPPORTEDTYPES(PCG_ALL_TYPES) void>;
+#undef PCG_ALL_TYPES
+
+			template <typename InputType, typename... InputTypes>
+			constexpr void Unroll(TypeHolder<InputType, InputTypes...> InTypeHolder)
+			{
+				ElementSize[MetadataTypes<InputType>::Id] = sizeof(InputType);
+				if constexpr (!std::is_trivially_constructible_v<InputType> || !std::is_trivially_destructible_v<InputType>)
+				{
+					ConstructFuncs[MetadataTypes<InputType>::Id] = &ConstructMemory<InputType>;
+					DestroyFuncs[MetadataTypes<InputType>::Id] = &DestructMemory<InputType>;
+				}
+
+				Unroll(TypeHolder<InputTypes...>{});
+			}
+
+			constexpr void Unroll(TypeHolder<void> InTypeHolder) {}
+
+		public:
+			constexpr FTypesFundamental()
+				: ElementSize{{}}
+				, ConstructFuncs{{}}
+				, DestroyFuncs{{}}
+			{
+				Unroll(AllTypes{});
+			}
+
+			size_t ElementSize[(uint8)EPCGMetadataTypes::Count];
+			FConstructMemoryFunc ConstructFuncs[(uint8)EPCGMetadataTypes::Count];
+			FDestroyMemoryFunc DestroyFuncs[(uint8)EPCGMetadataTypes::Count];
 		};
 
 		inline static constexpr FBroadcastableTypes BroadcastableTypes{};
 		inline static constexpr FConstructibleTypes ConstructibleTypes{};
+		inline static constexpr FTypesFundamental TypesFundamental{};
 
 		constexpr inline bool IsBroadcastable(uint16 FirstType, uint16 SecondType)
 		{
@@ -459,6 +581,23 @@ namespace PCG
 			return BroadcastableTypes[FirstType][SecondType];
 		}
 
+		inline bool IsBroadcastable(const FPCGMetadataAttributeDesc& FirstType, const FPCGMetadataAttributeDesc& SecondType)
+		{
+			// If it is the same type, it is always possible
+			if (FirstType.IsSameType(SecondType))
+			{
+				return true;
+			}
+
+			// Containers are not broadcastable if the types are not matching
+			if (!FirstType.ContainerTypes.IsEmpty() || !SecondType.ContainerTypes.IsEmpty())
+			{
+				return false;
+			}
+
+			return IsBroadcastable(static_cast<int16>(FirstType.ValueType), static_cast<int16>(SecondType.ValueType));
+		}
+
 		template <typename FirstType, typename SecondType>
 		constexpr inline bool IsBroadcastable()
 		{
@@ -475,7 +614,24 @@ namespace PCG
 
 			return ConstructibleTypes[FirstType][SecondType];
 		}
-		
+
+		inline bool IsConstructible(const FPCGMetadataAttributeDesc& FirstType, const FPCGMetadataAttributeDesc& SecondType)
+		{
+			// If it is the same type, it is always possible
+			if (FirstType.IsSameType(SecondType))
+			{
+				return true;
+			}
+
+			// Containers are not constructible if the types are not matching
+			if (!FirstType.ContainerTypes.IsEmpty() || !SecondType.ContainerTypes.IsEmpty())
+			{
+				return false;
+			}
+
+			return IsConstructible(static_cast<int16>(FirstType.ValueType), static_cast<int16>(SecondType.ValueType));
+		}
+
 		// Convenience function for accessor operation, to know if a Get will succeed with AllowBroadcast and AllowConstructible
 		// where the SecondType is the target type, and FirstType is the accessor type.
 		// ie. Accessor is an int, Accessor->Get<float>(...). First type is int, second type is float.
@@ -493,6 +649,27 @@ namespace PCG
 			return IsBroadcastable(FirstType, SecondType) || IsConstructible(SecondType, FirstType);
 		}
 
+		// Convenience function for accessor operation, to know if a Get will succeed with AllowBroadcast and AllowConstructible
+		// where the SecondType is the target type, and FirstType is the accessor type.
+		// ie. Accessor is an int, Accessor->Get<float>(...). First type is int, second type is float.
+		inline bool IsBroadcastableOrConstructible(const FPCGMetadataAttributeDesc& FirstType, const FPCGMetadataAttributeDesc& SecondType)
+		{
+			// If it is the same type, it is always possible
+			if (FirstType.IsSameType(SecondType))
+			{
+				return true;
+			}
+
+			// Containers are not broadcastable nor constructible if the types are not matching
+			if (!FirstType.ContainerTypes.IsEmpty() || !SecondType.ContainerTypes.IsEmpty())
+			{
+				return false;
+			}
+
+			return IsBroadcastable(static_cast<int16>(FirstType.ValueType), static_cast<int16>(SecondType.ValueType))
+				|| IsConstructible(static_cast<int16>(SecondType.ValueType), static_cast<int16>(FirstType.ValueType));
+		}
+
 		/**
 		* FirstType is more complex than SecondType, in the sense of PCG types, if both types are different and valid
 		* and we can broadcast SecondType to FirstType.
@@ -512,6 +689,56 @@ namespace PCG
 			return IsMoreComplexType(MetadataTypes<FirstType>::Id, MetadataTypes<SecondType>::Id);
 		}
 
+		inline size_t GetElementSize(uint16 TypeId)
+		{
+			if (TypeId >= static_cast<uint16>(EPCGMetadataTypes::Count))
+			{
+				return static_cast<size_t>(-1);
+			}
+
+			return TypesFundamental.ElementSize[TypeId];
+		}
+
+		inline FConstructMemoryFunc GetConstructMemoryFunc(uint16 TypeId)
+		{
+			if (TypeId >= static_cast<uint16>(EPCGMetadataTypes::Count))
+			{
+				return nullptr;
+			}
+
+			return TypesFundamental.ConstructFuncs[TypeId];
+		}
+
+		inline FDestroyMemoryFunc GetDestroyMemoryFunc(uint16 TypeId)
+		{
+			if (TypeId >= static_cast<uint16>(EPCGMetadataTypes::Count))
+			{
+				return nullptr;
+			}
+
+			return TypesFundamental.DestroyFuncs[TypeId];
+		}
+
+		inline FTransformFunc GetBroadcastTransformFunc(uint16 FirstType, uint16 SecondType)
+		{
+			if (FirstType >= static_cast<uint16>(EPCGMetadataTypes::Count) || SecondType >= static_cast<uint16>(EPCGMetadataTypes::Count))
+			{
+				return nullptr;
+			}
+
+			return BroadcastableTypes.TransformFuncs[FirstType][SecondType];
+		}
+
+		inline FTransformFunc GetConstructibleTransformFunc(uint16 FirstType, uint16 SecondType)
+		{
+			if (FirstType >= static_cast<uint16>(EPCGMetadataTypes::Count) || SecondType >= static_cast<uint16>(EPCGMetadataTypes::Count))
+			{
+				return nullptr;
+			}
+
+			return ConstructibleTypes.TransformFuncs[FirstType][SecondType];
+		}
+
 		/**
 		* Traits
 		* - CanCompare: Supports less than, greater than etc for ordering.
@@ -693,8 +920,8 @@ namespace PCG
 		{
 			using DistanceType = T;
 
-			static DistanceType Distance(const T& A, const T& B) 
-			{ 
+			static DistanceType Distance(const T& A, const T& B)
+			{
 				return FMath::Abs(A - B);
 			}
 		};
@@ -890,7 +1117,7 @@ namespace PCG
 				return A + B * Weight;
 			}
 
-			static DistanceType Distance(const T& A, const T& B) 
+			static DistanceType Distance(const T& A, const T& B)
 			{
 				return (A - B).Size();
 			}
@@ -1524,6 +1751,45 @@ namespace PCG
 			}
 		}
 
+		template <typename InType, typename OutType>
+		inline void GetValuesWithBroadcast_Internal(void* DestPtr, const void* SrcPtr, int32 Count)
+		{
+			TConstArrayView<InType> Src(static_cast<const InType*>(SrcPtr), Count);
+			TArrayView<OutType> Dest(static_cast<OutType*>(DestPtr), Count);
+
+			for (int32 i = 0; i < Count; ++i)
+			{
+				GetValueWithBroadcast<InType, OutType>(Src[i], Dest[i]);
+			}
+		}
+
+		template <typename InType, typename OutType>
+		inline void GetValuesWithConstructible_Internal(void* DestPtr, const void* SrcPtr, int32 Count)
+		{
+			static_assert(std::is_constructible_v<OutType, InType>);
+			TConstArrayView<InType> Src(static_cast<const InType*>(SrcPtr), Count);
+			TArrayView<OutType> Dest(static_cast<OutType*>(DestPtr), Count);
+
+			for (int32 i = 0; i < Count; ++i)
+			{
+				Dest[i] = OutType(Src[i]);
+			}
+		}
+
+		template <typename InType, typename OutType>
+		inline bool GetValueConstructible(const InType& InValue, OutType& OutValue)
+		{
+			if constexpr(std::is_constructible_v<OutType, InType>)
+			{
+				OutValue = OutType(InValue);
+				return true;
+			}
+			else
+			{
+				return false;
+			}
+		}
+
 		template<typename InType, typename OutType>
 		inline bool GetValueWithBroadcastAndConstructible(const InType& InValue, OutType& OutValue)
 		{
@@ -1534,21 +1800,13 @@ namespace PCG
 			}
 			else
 			{
-				// Special case'd because FSoftObjectPath currently has a deprecated constructor from FName which generates compile warnings.
-				constexpr bool bNameToSoftObjectPath = std::is_same_v<OutType, FSoftObjectPath>&& std::is_same_v<InType, FName>;
-
 				if constexpr (PCG::Private::IsBroadcastable<InType, OutType>())
 				{
 					return GetValueWithBroadcast<InType, OutType>(InValue, OutValue);
 				}
-				else if constexpr (std::is_constructible_v<OutType, InType> && !bNameToSoftObjectPath)
-				{
-					OutValue = OutType(InValue);
-					return true;
-				}
 				else
 				{
-					return false;
+					return GetValueConstructible<InType, OutType>(InValue, OutValue);
 				}
 			}
 		}
diff --git a/Engine/Plugins/PCG/Source/PCG/Public/Metadata/PCGMetadataCommon.h b/Engine/Plugins/PCG/Source/PCG/Public/Metadata/PCGMetadataCommon.h
--- a/Engine/Plugins/PCG/Source/PCG/Public/Metadata/PCGMetadataCommon.h
+++ b/Engine/Plugins/PCG/Source/PCG/Public/Metadata/PCGMetadataCommon.h
@@ -1,9 +1,13 @@
 // Copyright Epic Games, Inc. All Rights Reserved.
 
 #pragma once
+#include "Misc/TVariant.h"
 
 #include "PCGMetadataCommon.generated.h"
 
+class FScriptMapHelper;
+class FScriptSetHelper;
+
 typedef int64 PCGMetadataEntryKey;
 typedef int32 PCGMetadataAttributeKey;
 typedef int32 PCGMetadataValueKey;
@@ -293,6 +297,74 @@ struct FPCGMetadataAttributeDesc
 
 	PCG_API friend int32 GetTypeHash(const FPCGMetadataAttributeDesc& Desc);
 	
-	UE_DEPRECATED(all, "Should not be used anymore as some information is lost in the process.")
-	PCG_API int16 AsLegacyTypeId() const;
+	// For already predefined structs, replace `EPCGMetadataTypes::Struct` by its predefined type.
+	PCG_API void FixLegacyTypeId();
 };
+
+/**
+ * Variants to type erase chunk of data to be read/write in attributes or accessors
+ */
+namespace PCG::Private
+{
+	//////////////////////
+	// OutValues Section
+	//////////////////////
+	struct FOutValuesByValue
+	{
+		void* OutValues;
+	};
+	
+	struct FOutValuesByPtr
+	{
+		TArrayView<const void*> OutValues;
+	};
+	
+	struct FOutValuesAsArray
+	{
+		TArrayView<TTuple<const void*, int32>> OutValues;
+	};
+	
+	struct FOutValuesAsSet
+	{
+		TArrayView<FScriptSetHelper*> OutValues;
+	};
+	
+	struct FOutValuesAsMap
+	{
+		TArrayView<FScriptMapHelper*> OutValues;
+	};
+	
+	using FOutValues = TVariant<FOutValuesByValue, FOutValuesByPtr, FOutValuesAsArray, FOutValuesAsSet, FOutValuesAsMap>;
+	
+	//////////////////////
+	// InValues Section
+	//////////////////////
+	struct FInValuesByValue
+	{
+		const void* InValues;
+		int32 Count;
+	};
+
+	struct FInValuesByValueSubset
+	{
+		const void* InValues;
+		TArray<int32> Indices;
+	};
+
+	struct FInValuesAsArray
+	{
+		TArrayView<TTuple<const void*, int32>> InValues;
+	};
+
+	struct FInValuesAsSet
+	{
+		TArrayView<TArray<const void*>> InValues;
+	};
+
+	struct FInValuesAsMap
+	{
+		TArrayView<TArray<TPair<const void*, const void*>>> InValues;
+	};
+
+	using FInValues = TVariant<FInValuesByValue, FInValuesAsArray, FInValuesAsSet, FInValuesAsMap, FInValuesByValueSubset>;
+}
\ No newline at end of file
diff --git a/Engine/Plugins/PCG/Source/PCG/Public/Metadata/PCGMetadataContainerTypes.h b/Engine/Plugins/PCG/Source/PCG/Public/Metadata/PCGMetadataContainerTypes.h
--- a/Engine/Plugins/PCG/Source/PCG/Public/Metadata/PCGMetadataContainerTypes.h
+++ b/Engine/Plugins/PCG/Source/PCG/Public/Metadata/PCGMetadataContainerTypes.h
@@ -5,6 +5,7 @@
 #include "UObject/UnrealType.h"
 
 class FPCGMetadataAttributeGeneric;
+class IPCGAttributeAccessor;
 
 namespace PCG
 {
@@ -93,6 +94,7 @@ public:
 	
 private:
 	friend FPCGMetadataAttributeGeneric;
+	friend IPCGAttributeAccessor;
 	FScriptSetHelper Helper;
 };
 	
@@ -180,6 +182,7 @@ public:
 	
 private:
 	friend FPCGMetadataAttributeGeneric;
+	friend IPCGAttributeAccessor;
 	FScriptMapHelper Helper;
 };
 }
diff --git a/Engine/Plugins/PCG/Tests/Private/Metadata/Accessors/PCGAttributeAccessorTests.cpp b/Engine/Plugins/PCG/Tests/Private/Metadata/Accessors/PCGAttributeAccessorTests.cpp
new file mode 100644
--- /dev/null
+++ b/Engine/Plugins/PCG/Tests/Private/Metadata/Accessors/PCGAttributeAccessorTests.cpp
@@ -0,0 +1,885 @@
+// Copyright Epic Games, Inc. All Rights Reserved.
+
+// =============================================================================
+// PCG Attribute Accessor Tests
+//
+// Tests for IPCGAttributeAccessor Get/Set operations on PCG metadata attributes.
+// Uses float, FVector, and FString attributes with entry-based and point-based keys.
+//
+// Test summary:
+//  - SingleGetDefault:                  Single Get on default (invalid) entry key returns attribute default values
+//  - SingleGet:                         Single Get on valid entry keys returns correct per-entry values
+//  - SingleGetOutsideRange:             Single Get with out-of-range indices wraps around
+//  - GetRangeDefault:                   GetRange on default entry key fills arrays with default values
+//  - GetRange:                          GetRange on valid entry keys fills arrays with correct values
+//  - GetRangeOutsideRange:              GetRange with out-of-range indices wraps around
+//  - SingleSetDefault:                  Single Set on default entry key updates the attribute default value
+//  - SingleSet:                         Single Set on valid entry keys writes values to specific entries
+//  - SingleSetInvalidKey:               Single Set on invalid entry key creates a new entry instead of modifying default
+//  - SetRange:                          SetRange writes an array of values to consecutive entries
+//  - SetRangeOutsideRange:              SetRange with out-of-range indices or too many values fails
+//  - SingleGetPoints:                   Single Get using point-based keys resolves metadata entry indirection
+//  - GetRangePoints:                    GetRange using point-based keys respects reverse entry order
+//  - SingleSetPoints:                   Single Set using point-based keys writes through entry indirection
+//  - SetRangePoints:                    SetRange using point-based keys writes through entry indirection
+//  - GetWithConstructibleType:          Get/Set with constructible type conversion (float<->double) using AllowConstructible flag
+//  - GetWithBroadcastableType:          Get/Set with broadcastable type conversion (float->FVector/FString) using AllowBroadcast flag
+//  - SoftObjectPathToStringWithStrictType: FSoftObjectPath/FSoftClassPath to FString conversion with StrictType flag
+// =============================================================================
+
+#include "PCGTestsCommon.h"
+
+#include "Data/PCGPointData.h"
+#include "Metadata/Accessors/IPCGAttributeAccessorTpl.h"
+#include "Metadata/Accessors/PCGAttributeAccessor.h"
+#include "Metadata/Accessors/PCGAttributeAccessorHelpers.h"
+
+#include <catch2/catch_test_macros.hpp>
+
+#include "TestHarness.h"
+
+namespace PCGAttributeAccessorTests
+{
+	constexpr int32 ValuesSize = 3;
+
+	const float FloatValues[ValuesSize] = { 0.57441f, -0.289f, 0.1587f};
+	const FVector VectorValues[ValuesSize] = { FVector(1.024f, 0.445f, 0.587f), FVector(4.6470f, 3.19874f, 9.6648f), FVector(4.690f, 7.874f, 8.668f) };
+	const FString StringValues[ValuesSize] = { FString(TEXT("Str1")), FString(TEXT("Foo")), FString(TEXT("MyStr")) };
+
+	template <typename T, size_t N>
+	TArrayView<T> MakeArrayView(T(& Array)[N])
+	{
+		return TArrayView<T>(Array, N);
+	}
+
+	template <typename T, size_t N>
+	TArrayView<const T> MakeConstArrayView(T(&Array)[N])
+	{
+		return TArrayView<const T>(Array, N);
+	}
+
+	struct AttributeData
+	{
+		explicit AttributeData(UPCGBasePointData* InPointData)
+			: PointData(InPointData)
+		{
+			FloatAttribute = PointData->Metadata->CreateAttribute<float>(TEXT("FloatAttribute"), 0.0f, true, false);
+			VectorAttribute = PointData->Metadata->CreateAttribute<FVector>(TEXT("VectorAttribute"), FVector::ZeroVector, false, false);
+			StringAttribute = PointData->Metadata->CreateAttribute<FString>(TEXT("StringAttribute"), FString(), false, false);
+
+			// All points will point to a given entry, in reverse order, except the last one, that will point to DefaultValue
+			PointData->SetNumPoints(ValuesSize + 1);
+			TPCGValueRange<int64> MetadataEntryRange = PointData->GetMetadataEntryValueRange();
+
+			for (int32 i = 0; i < ValuesSize; ++i)
+			{
+				PCGMetadataEntryKey EntryKey = PointData->Metadata->AddEntry();
+				FloatAttribute->SetValue(EntryKey, FloatValues[i]);
+				VectorAttribute->SetValue(EntryKey, VectorValues[i]);
+				StringAttribute->SetValue(EntryKey, StringValues[i]);
+				MetadataEntryRange[ValuesSize - i - 1] = EntryKey;
+			}
+
+			FloatAccessor = PCGAttributeAccessorHelpers::CreateAccessor(FloatAttribute, PointData->Metadata);
+			VectorAccessor = PCGAttributeAccessorHelpers::CreateAccessor(VectorAttribute, PointData->Metadata);
+			StringAccessor = PCGAttributeAccessorHelpers::CreateAccessor(StringAttribute, PointData->Metadata);
+		}
+
+		void ResetTempFloats()
+		{
+			TempFloats[0] = 1.0f;
+			TempFloats[1] = 2.0f;
+			TempFloats[2] = 3.0f;
+		}
+
+		TObjectPtr<UPCGBasePointData> PointData;
+		FPCGMetadataAttribute<float>* FloatAttribute;
+		FPCGMetadataAttribute<FVector>* VectorAttribute;
+		FPCGMetadataAttribute<FString>* StringAttribute;
+
+		TUniquePtr<IPCGAttributeAccessor> FloatAccessor;
+		TUniquePtr<IPCGAttributeAccessor> VectorAccessor;
+		TUniquePtr<IPCGAttributeAccessor> StringAccessor;
+
+		float TempFloats[ValuesSize] = { 1.0f, 2.0f, 3.0f };
+		FVector TempVectors[ValuesSize] = { FVector(0.1f, 0.1f, 0.1f), FVector(0.1f, 0.2f, 0.1f), FVector(0.1f, 0.2f, 0.3f) };
+		FString TempStrings[ValuesSize] = { TEXT("Bla"), TEXT("BlaBla"), TEXT("Blou") };
+		int32 TempInvalidInts[ValuesSize] = { 5, 6, 7 };
+	};
+}
+
+// Tests that getting a single value using the default (invalid) entry key returns the attribute's default value for each type (float, vector, string), and fails when requesting an int from a float accessor.
+TEST_CASE_METHOD(PCGTests::FPCGBaseTestWithBasicContext, "PCG::Accessor::Attribute::SingleGetDefault", "[PCG][Accessor][Attribute]")
+{
+	using namespace PCGAttributeAccessorTests;
+
+	AttributeData Data(CreatePointData());
+	FPCGAttributeAccessorKeysEntries DefaultEntry(PCGInvalidEntryKey);
+
+	REQUIRE_EQUAL(DefaultEntry.GetNum(), 1);
+
+	SECTION("Get default float attribute")
+	{
+		REQUIRE(Data.FloatAccessor->Get(*Data.TempFloats, DefaultEntry));
+		REQUIRE_EQUAL(*Data.TempFloats, 0.0f);
+	}
+
+	SECTION("Get default vector attribute")
+	{
+		REQUIRE(Data.VectorAccessor->Get(*Data.TempVectors, DefaultEntry));
+		REQUIRE_EQUAL(*Data.TempVectors, FVector::ZeroVector);
+	}
+
+	SECTION("Get default string attribute")
+	{
+		REQUIRE(Data.StringAccessor->Get(*Data.TempStrings, DefaultEntry));
+		REQUIRE_EQUAL(*Data.TempStrings, FString{});
+	}
+
+	SECTION("Get default float with int fails")
+	{
+		REQUIRE_FALSE(Data.FloatAccessor->Get(*Data.TempInvalidInts, DefaultEntry));
+		REQUIRE_EQUAL(*Data.TempInvalidInts, 5);
+	}
+}
+
+// Tests that getting a single value by entry index returns the correct stored value for float, vector, and string attributes, and fails when requesting an int from a float accessor.
+TEST_CASE_METHOD(PCGTests::FPCGBaseTestWithBasicContext, "PCG::Accessor::Attribute::SingleGet", "[PCG][Accessor][Attribute]")
+{
+	using namespace PCGAttributeAccessorTests;
+
+	AttributeData Data(CreatePointData());
+	FPCGAttributeAccessorKeysEntries Keys(Data.FloatAttribute->GetMetadata());
+
+	REQUIRE_EQUAL(Keys.GetNum(), ValuesSize);
+
+	for (int32 i = 0; i < ValuesSize; ++i)
+	{
+		REQUIRE(Data.FloatAccessor->Get(Data.TempFloats[i], i, Keys));
+		REQUIRE_EQUAL(Data.TempFloats[i], FloatValues[i]);
+
+		REQUIRE(Data.VectorAccessor->Get(Data.TempVectors[i], i, Keys));
+		REQUIRE_EQUAL(Data.TempVectors[i], VectorValues[i]);
+
+		REQUIRE(Data.StringAccessor->Get(Data.TempStrings[i], i, Keys));
+		REQUIRE_EQUAL(Data.TempStrings[i], StringValues[i]);
+
+		REQUIRE_FALSE(Data.FloatAccessor->Get(Data.TempInvalidInts[i], i, Keys));
+		REQUIRE_EQUAL(Data.TempInvalidInts[i], 5 + i);
+	}
+}
+
+// Tests that getting a single value with an index beyond the entry count wraps around using modulo.
+TEST_CASE_METHOD(PCGTests::FPCGBaseTestWithBasicContext, "PCG::Accessor::Attribute::SingleGetOutsideRange", "[PCG][Accessor][Attribute]")
+{
+	using namespace PCGAttributeAccessorTests;
+
+	AttributeData Data(CreatePointData());
+	FPCGAttributeAccessorKeysEntries Keys(Data.FloatAttribute->GetMetadata());
+
+	float DefaultFloat = 1.0f;
+
+	// 3 -> 7 wrapping to 0 -> 3 and back to 0
+	for (int32 i = 0; i < ValuesSize + 1; ++i)
+	{
+		REQUIRE(Data.FloatAccessor->Get(DefaultFloat, i + ValuesSize, Keys));
+		REQUIRE_EQUAL(DefaultFloat, FloatValues[i % ValuesSize]);
+	}
+}
+
+// Tests that GetRange using the default (invalid) entry key fills the output array with default values for each type, fails for int from float accessor, and works correctly with a sub-range starting at index 1.
+TEST_CASE_METHOD(PCGTests::FPCGBaseTestWithBasicContext, "PCG::Accessor::Attribute::GetRangeDefault", "[PCG][Accessor][Attribute]")
+{
+	using namespace PCGAttributeAccessorTests;
+
+	AttributeData Data(CreatePointData());
+	FPCGAttributeAccessorKeysEntries DefaultEntry(PCGInvalidEntryKey);
+
+	SECTION("Full range")
+	{
+		int32 Index = 0;
+
+		SECTION("Float")
+		{
+			REQUIRE(Data.FloatAccessor->GetRange(MakeArrayView(Data.TempFloats), Index, DefaultEntry));
+			REQUIRE_EQUAL(Data.TempFloats[0], 0.0f);
+			REQUIRE_EQUAL(Data.TempFloats[1], 0.0f);
+			REQUIRE_EQUAL(Data.TempFloats[2], 0.0f);
+		}
+
+		SECTION("Vector")
+		{
+			REQUIRE(Data.VectorAccessor->GetRange(MakeArrayView(Data.TempVectors), Index, DefaultEntry));
+			REQUIRE_EQUAL(Data.TempVectors[0], FVector::ZeroVector);
+			REQUIRE_EQUAL(Data.TempVectors[1], FVector::ZeroVector);
+			REQUIRE_EQUAL(Data.TempVectors[2], FVector::ZeroVector);
+		}
+
+		SECTION("String")
+		{
+			REQUIRE(Data.StringAccessor->GetRange(MakeArrayView(Data.TempStrings), Index, DefaultEntry));
+			REQUIRE_EQUAL(Data.TempStrings[0], FString{});
+			REQUIRE_EQUAL(Data.TempStrings[1], FString{});
+			REQUIRE_EQUAL(Data.TempStrings[2], FString{});
+		}
+
+		SECTION("Invalid int type")
+		{
+			REQUIRE_FALSE(Data.FloatAccessor->GetRange(MakeArrayView(Data.TempInvalidInts), Index, DefaultEntry));
+			REQUIRE_EQUAL(Data.TempInvalidInts[0], 5);
+			REQUIRE_EQUAL(Data.TempInvalidInts[1], 6);
+			REQUIRE_EQUAL(Data.TempInvalidInts[2], 7);
+		}
+	}
+
+	SECTION("Smaller range from index 1")
+	{
+		int32 Index = 1;
+
+		REQUIRE(Data.FloatAccessor->GetRange(TArrayView<float>(Data.TempFloats, 2), Index, DefaultEntry));
+		REQUIRE_EQUAL(Data.TempFloats[0], 0.0f);
+		REQUIRE_EQUAL(Data.TempFloats[1], 0.0f);
+		REQUIRE_EQUAL(Data.TempFloats[2], 3.0f);
+	}
+}
+
+// Tests that GetRange on valid entry keys fills the output array with correct stored values for float, vector, and string, fails for int from float accessor, and works with a sub-range starting at index 1.
+TEST_CASE_METHOD(PCGTests::FPCGBaseTestWithBasicContext, "PCG::Accessor::Attribute::GetRange", "[PCG][Accessor][Attribute]")
+{
+	using namespace PCGAttributeAccessorTests;
+
+	AttributeData Data(CreatePointData());
+	FPCGAttributeAccessorKeysEntries Keys(Data.FloatAttribute->GetMetadata());
+
+	SECTION("Full range")
+	{
+		int32 Index = 0;
+
+		SECTION("Float")
+		{
+			REQUIRE(Data.FloatAccessor->GetRange(MakeArrayView(Data.TempFloats), Index, Keys));
+			REQUIRE_EQUAL(Data.TempFloats[0], FloatValues[0]);
+			REQUIRE_EQUAL(Data.TempFloats[1], FloatValues[1]);
+			REQUIRE_EQUAL(Data.TempFloats[2], FloatValues[2]);
+		}
+
+		SECTION("Vector")
+		{
+			REQUIRE(Data.VectorAccessor->GetRange(MakeArrayView(Data.TempVectors), Index, Keys));
+			REQUIRE_EQUAL(Data.TempVectors[0], VectorValues[0]);
+			REQUIRE_EQUAL(Data.TempVectors[1], VectorValues[1]);
+			REQUIRE_EQUAL(Data.TempVectors[2], VectorValues[2]);
+		}
+
+		SECTION("String")
+		{
+			REQUIRE(Data.StringAccessor->GetRange(MakeArrayView(Data.TempStrings), Index, Keys));
+			REQUIRE_EQUAL(Data.TempStrings[0], StringValues[0]);
+			REQUIRE_EQUAL(Data.TempStrings[1], StringValues[1]);
+			REQUIRE_EQUAL(Data.TempStrings[2], StringValues[2]);
+		}
+
+		SECTION("Invalid int type")
+		{
+			REQUIRE_FALSE(Data.FloatAccessor->GetRange(MakeArrayView(Data.TempInvalidInts), Index, Keys));
+			REQUIRE_EQUAL(Data.TempInvalidInts[0], 5);
+			REQUIRE_EQUAL(Data.TempInvalidInts[1], 6);
+			REQUIRE_EQUAL(Data.TempInvalidInts[2], 7);
+		}
+	}
+
+	SECTION("Smaller range from index 1")
+	{
+		int32 Index = 1;
+
+		REQUIRE(Data.FloatAccessor->GetRange(TArrayView<float>(Data.TempFloats, 2), Index, Keys));
+		REQUIRE_EQUAL(Data.TempFloats[0], FloatValues[1]);
+		REQUIRE_EQUAL(Data.TempFloats[1], FloatValues[2]);
+		REQUIRE_EQUAL(Data.TempFloats[2], 3.0f);
+	}
+}
+
+// Tests that GetRange with out-of-range indices wraps around correctly, including when requesting more values than entries exist.
+TEST_CASE_METHOD(PCGTests::FPCGBaseTestWithBasicContext, "PCG::Accessor::Attribute::GetRangeOutsideRange", "[PCG][Accessor][Attribute]")
+{
+	using namespace PCGAttributeAccessorTests;
+
+	AttributeData Data(CreatePointData());
+	FPCGAttributeAccessorKeysEntries Keys(Data.FloatAttribute->GetMetadata());
+
+	SECTION("Wrapping 3 -> 6 to 0 -> 3")
+	{
+		int32 OutsideRangeIndex = ValuesSize;
+
+		REQUIRE(Data.FloatAccessor->GetRange(MakeArrayView(Data.TempFloats), OutsideRangeIndex, Keys));
+		REQUIRE_EQUAL(Data.TempFloats[0], FloatValues[0]);
+		REQUIRE_EQUAL(Data.TempFloats[1], FloatValues[1]);
+		REQUIRE_EQUAL(Data.TempFloats[2], FloatValues[2]);
+	}
+
+	SECTION("Wrapping 4, 5, 7 to 1, 2, 0")
+	{
+		int32 OutsideRangeIndex = ValuesSize + 1;
+
+		REQUIRE(Data.FloatAccessor->GetRange(MakeArrayView(Data.TempFloats), OutsideRangeIndex, Keys));
+		REQUIRE_EQUAL(Data.TempFloats[0], FloatValues[1]);
+		REQUIRE_EQUAL(Data.TempFloats[1], FloatValues[2]);
+		REQUIRE_EQUAL(Data.TempFloats[2], FloatValues[0]);
+	}
+
+	SECTION("Gathering more values wraps")
+	{
+		constexpr int32 MoreFloatValuesSize = ValuesSize * 3;
+		float MoreFloatValues[MoreFloatValuesSize];
+		for (int32 i = 0; i < MoreFloatValuesSize; ++i)
+		{
+			MoreFloatValues[i] = float(i);
+		}
+
+		REQUIRE(Data.FloatAccessor->GetRange(MakeArrayView(MoreFloatValues), 0, Keys));
+
+		for (int32 i = 0; i < MoreFloatValuesSize; ++i)
+		{
+			REQUIRE_EQUAL(MoreFloatValues[i], FloatValues[i % ValuesSize]);
+		}
+	}
+}
+
+// Tests that setting a single value on the default entry key with AllowSetDefaultValue flag updates the attribute's default value for float, vector, and string.
+TEST_CASE_METHOD(PCGTests::FPCGBaseTestWithBasicContext, "PCG::Accessor::Attribute::SingleSetDefault", "[PCG][Accessor][Attribute]")
+{
+	using namespace PCGAttributeAccessorTests;
+
+	AttributeData Data(CreatePointData());
+	FPCGAttributeAccessorKeysEntries DefaultEntry(PCGInvalidEntryKey);
+
+	SECTION("Float")
+	{
+		float NewDefaultFloat = 0.5f;
+		REQUIRE(Data.FloatAccessor->Set(NewDefaultFloat, DefaultEntry, EPCGAttributeAccessorFlags::AllowSetDefaultValue));
+		REQUIRE_EQUAL(Data.FloatAttribute->GetValue(PCGDefaultValueKey), NewDefaultFloat);
+	}
+
+	SECTION("Vector")
+	{
+		FVector NewDefaultVector = FVector(0.5f, 0.2f, 0.1f);
+		REQUIRE(Data.VectorAccessor->Set(NewDefaultVector, DefaultEntry, EPCGAttributeAccessorFlags::AllowSetDefaultValue));
+		REQUIRE_EQUAL(Data.VectorAttribute->GetValue(PCGDefaultValueKey), NewDefaultVector);
+	}
+
+	SECTION("String")
+	{
+		FString NewDefaultString(TEXT("AAA"));
+		REQUIRE(Data.StringAccessor->Set(NewDefaultString, DefaultEntry, EPCGAttributeAccessorFlags::AllowSetDefaultValue));
+		REQUIRE_EQUAL(Data.StringAttribute->GetValue(PCGDefaultValueKey), NewDefaultString);
+	}
+}
+
+// Tests that setting a single value at a specific entry index writes correctly for float, vector, and string, and fails when the index is outside the valid range.
+TEST_CASE_METHOD(PCGTests::FPCGBaseTestWithBasicContext, "PCG::Accessor::Attribute::SingleSet", "[PCG][Accessor][Attribute]")
+{
+	using namespace PCGAttributeAccessorTests;
+
+	AttributeData Data(CreatePointData());
+	FPCGAttributeAccessorKeysEntries Keys(const_cast<FPCGMetadataDomain*>(Data.FloatAttribute->GetMetadataDomain()));
+
+	SECTION("Set float attribute")
+	{
+		REQUIRE(Data.FloatAccessor->Set(*Data.TempFloats, Keys));
+		REQUIRE_EQUAL(Data.FloatAttribute->GetValueFromItemKey(0), *Data.TempFloats);
+	}
+
+	SECTION("Set vector attribute at index 1")
+	{
+		REQUIRE(Data.VectorAccessor->Set(*Data.TempVectors, /*Index=*/ 1, Keys));
+		REQUIRE_EQUAL(Data.VectorAttribute->GetValueFromItemKey(1), *Data.TempVectors);
+	}
+
+	SECTION("Set string attribute at index 2")
+	{
+		REQUIRE(Data.StringAccessor->Set(*Data.TempStrings, /*Index=*/ 2, Keys));
+		REQUIRE_EQUAL(Data.StringAttribute->GetValueFromItemKey(2), *Data.TempStrings);
+	}
+
+	SECTION("Set float attribute outside range is invalid")
+	{
+		float InvalidFloat = 2.147f;
+		REQUIRE_FALSE(Data.FloatAccessor->Set(InvalidFloat, /*Index=*/ 3, Keys));
+	}
+}
+
+// Tests that setting a value on an invalid entry key (without AllowSetDefaultValue) creates a new metadata entry rather than modifying the default value.
+TEST_CASE_METHOD(PCGTests::FPCGBaseTestWithBasicContext, "PCG::Accessor::Attribute::SingleSetInvalidKey", "[PCG][Accessor][Attribute]")
+{
+	using namespace PCGAttributeAccessorTests;
+
+	AttributeData Data(CreatePointData());
+	FPCGAttributeAccessorKeysEntries DefaultEntry(PCGInvalidEntryKey);
+
+	float NewDefaultFloat = 0.5f;
+
+	REQUIRE(Data.FloatAccessor->Set(NewDefaultFloat, DefaultEntry));
+	REQUIRE_EQUAL(Data.FloatAttribute->GetValue(PCGDefaultValueKey), 0.0f);
+	REQUIRE_EQUAL(Data.PointData->Metadata->GetItemCountForChild(), ValuesSize + 1);
+	REQUIRE_EQUAL(Data.FloatAttribute->GetValueFromItemKey(ValuesSize), NewDefaultFloat);
+}
+
+// Tests that SetRange writes an array of values to consecutive entries for float, vector, and string attributes.
+TEST_CASE_METHOD(PCGTests::FPCGBaseTestWithBasicContext, "PCG::Accessor::Attribute::SetRange", "[PCG][Accessor][Attribute]")
+{
+	using namespace PCGAttributeAccessorTests;
+
+	AttributeData Data(CreatePointData());
+	FPCGAttributeAccessorKeysEntries Keys(const_cast<FPCGMetadataDomain*>(Data.FloatAttribute->GetMetadataDomain()));
+
+	const int32 Index = 0;
+
+	SECTION("Float")
+	{
+		REQUIRE(Data.FloatAccessor->SetRange(MakeConstArrayView(Data.TempFloats), Index, Keys));
+		REQUIRE_EQUAL(Data.FloatAttribute->GetValueFromItemKey(0), Data.TempFloats[0]);
+		REQUIRE_EQUAL(Data.FloatAttribute->GetValueFromItemKey(1), Data.TempFloats[1]);
+		REQUIRE_EQUAL(Data.FloatAttribute->GetValueFromItemKey(2), Data.TempFloats[2]);
+	}
+
+	SECTION("Vector")
+	{
+		REQUIRE(Data.VectorAccessor->SetRange(MakeConstArrayView(Data.TempVectors), Index, Keys));
+		REQUIRE_EQUAL(Data.VectorAttribute->GetValueFromItemKey(0), Data.TempVectors[0]);
+		REQUIRE_EQUAL(Data.VectorAttribute->GetValueFromItemKey(1), Data.TempVectors[1]);
+		REQUIRE_EQUAL(Data.VectorAttribute->GetValueFromItemKey(2), Data.TempVectors[2]);
+	}
+
+	SECTION("String")
+	{
+		REQUIRE(Data.StringAccessor->SetRange(MakeConstArrayView(Data.TempStrings), Index, Keys));
+		REQUIRE_EQUAL(Data.StringAttribute->GetValueFromItemKey(0), Data.TempStrings[0]);
+		REQUIRE_EQUAL(Data.StringAttribute->GetValueFromItemKey(1), Data.TempStrings[1]);
+		REQUIRE_EQUAL(Data.StringAttribute->GetValueFromItemKey(2), Data.TempStrings[2]);
+	}
+}
+
+// Tests that SetRange fails when providing too many values, when the index offset is outside the valid range, or when the index plus array size exceeds the range.
+TEST_CASE_METHOD(PCGTests::FPCGBaseTestWithBasicContext, "PCG::Accessor::Attribute::SetRangeOutsideRange", "[PCG][Accessor][Attribute]")
+{
+	using namespace PCGAttributeAccessorTests;
+
+	AttributeData Data(CreatePointData());
+	FPCGAttributeAccessorKeysEntries Keys(Data.FloatAttribute->GetMetadata());
+
+	const int32 Index = 0;
+
+	SECTION("Too many values")
+	{
+		float TooManyFloats[ValuesSize + 1] = { 1.0f, 2.0f, 3.0f, 4.0f };
+		REQUIRE_FALSE(Data.FloatAccessor->SetRange(MakeConstArrayView(TooManyFloats), Index, Keys));
+	}
+
+	SECTION("Index offset outside range")
+	{
+		REQUIRE_FALSE(Data.FloatAccessor->SetRange(MakeConstArrayView(Data.TempFloats), Index + 999, Keys));
+	}
+
+	SECTION("Index offset valid but too many values")
+	{
+		REQUIRE_FALSE(Data.FloatAccessor->SetRange(MakeConstArrayView(Data.TempFloats), Index + 1, Keys));
+	}
+}
+
+// Tests that getting a single value using point-based keys correctly resolves the metadata entry indirection and returns the value associated with each point.
+TEST_CASE_METHOD(PCGTests::FPCGBaseTestWithBasicContext, "PCG::Accessor::Attribute::SingleGetPoints", "[PCG][Accessor][Attribute]")
+{
+	using namespace PCGAttributeAccessorTests;
+
+	AttributeData Data(CreatePointData());
+
+	const TConstPCGValueRange<int64> MetadataEntryRange = Data.PointData->GetConstMetadataEntryValueRange();
+	TUniquePtr<const IPCGAttributeAccessorKeys> Keys = PCGAttributeAccessorHelpers::CreateConstKeys(Data.PointData, FPCGAttributePropertySelector());
+
+	REQUIRE_EQUAL(Keys->GetNum(), Data.PointData->GetNumPoints());
+	REQUIRE(Keys->IsReadOnly());
+
+	for (int32 i = 0; i < 4; ++i)
+	{
+		REQUIRE(Data.FloatAccessor->Get(*Data.TempFloats, i, *Keys));
+		REQUIRE_EQUAL(*Data.TempFloats, Data.FloatAttribute->GetValueFromItemKey(MetadataEntryRange[i]));
+	}
+}
+
+// Tests that GetRange using point-based keys returns values in point order (reverse entry order), with the last point returning the default value.
+TEST_CASE_METHOD(PCGTests::FPCGBaseTestWithBasicContext, "PCG::Accessor::Attribute::GetRangePoints", "[PCG][Accessor][Attribute]")
+{
+	using namespace PCGAttributeAccessorTests;
+
+	AttributeData Data(CreatePointData());
+
+	const TConstPCGValueRange<int64> MetadataEntryRange = Data.PointData->GetConstMetadataEntryValueRange();
+	TUniquePtr<const IPCGAttributeAccessorKeys> Keys = PCGAttributeAccessorHelpers::CreateConstKeys(Data.PointData, FPCGAttributePropertySelector());
+
+	REQUIRE(Keys->IsReadOnly());
+
+	float TempFloats[4] = { 1.0f, 2.0f, 3.0f, 4.0f };
+	const int32 Index = 0;
+
+	// Points entry are in reverse order, and last one is the default value.
+	REQUIRE(Data.FloatAccessor->GetRange(MakeArrayView(TempFloats), Index, *Keys));
+	REQUIRE_EQUAL(TempFloats[0], FloatValues[2]);
+	REQUIRE_EQUAL(TempFloats[1], FloatValues[1]);
+	REQUIRE_EQUAL(TempFloats[2], FloatValues[0]);
+	REQUIRE_EQUAL(TempFloats[3], 0.0f);
+}
+
+// Tests that setting a single value using point-based keys writes through the metadata entry indirection, preserves existing entry keys, and allocates a new entry for a point that had the default entry.
+TEST_CASE_METHOD(PCGTests::FPCGBaseTestWithBasicContext, "PCG::Accessor::Attribute::SingleSetPoints", "[PCG][Accessor][Attribute]")
+{
+	using namespace PCGAttributeAccessorTests;
+
+	AttributeData Data(CreatePointData());
+
+	const TConstPCGValueRange<int64> MetadataEntryRange = Data.PointData->GetConstMetadataEntryValueRange();
+	TUniquePtr<IPCGAttributeAccessorKeys> Keys = PCGAttributeAccessorHelpers::CreateKeys(Data.PointData, FPCGAttributePropertySelector());
+
+	REQUIRE_EQUAL(Keys->GetNum(), Data.PointData->GetNumPoints());
+	REQUIRE_FALSE(Keys->IsReadOnly());
+
+	float TempFloats[4] = { 1.0f, 2.0f, 3.0f, 4.0f };
+	const int32 Index = 0;
+
+	for (int32 i = 0; i < 4; ++i)
+	{
+		REQUIRE(Data.FloatAccessor->Set(TempFloats[i], i, *Keys));
+		REQUIRE_EQUAL(Data.FloatAttribute->GetValueFromItemKey(MetadataEntryRange[i]), TempFloats[i]);
+	}
+
+	// Points should not have a different entry key, if it was not invalid in the first place
+	REQUIRE_EQUAL(MetadataEntryRange[0], PCGMetadataEntryKey(2));
+	REQUIRE_EQUAL(MetadataEntryRange[1], PCGMetadataEntryKey(1));
+	REQUIRE_EQUAL(MetadataEntryRange[2], PCGFirstEntryKey);
+
+	// Also last point was a default entry. Making sure that we have one more entry and was set correctly
+	REQUIRE_EQUAL(Data.PointData->Metadata->GetItemCountForChild(), ValuesSize + 1);
+	REQUIRE_EQUAL(MetadataEntryRange[3], PCGMetadataEntryKey(3));
+}
+
+// Tests that SetRange using point-based keys writes values through entry indirection, preserves existing entry keys, and allocates a new entry for a point that had the default entry.
+TEST_CASE_METHOD(PCGTests::FPCGBaseTestWithBasicContext, "PCG::Accessor::Attribute::SetRangePoints", "[PCG][Accessor][Attribute]")
+{
+	using namespace PCGAttributeAccessorTests;
+
+	AttributeData Data(CreatePointData());
+
+	const TConstPCGValueRange<int64> MetadataEntryRange = Data.PointData->GetConstMetadataEntryValueRange();
+	TUniquePtr<IPCGAttributeAccessorKeys> Keys = PCGAttributeAccessorHelpers::CreateKeys(Data.PointData, FPCGAttributePropertySelector());
+
+	float TempFloats[4] = { 1.0f, 2.0f, 3.0f, 4.0f };
+	const int32 Index = 0;
+
+	REQUIRE(Data.FloatAccessor->SetRange(MakeConstArrayView(TempFloats), Index, *Keys));
+	REQUIRE_EQUAL(Data.FloatAttribute->GetValueFromItemKey(MetadataEntryRange[0]), TempFloats[0]);
+	REQUIRE_EQUAL(Data.FloatAttribute->GetValueFromItemKey(MetadataEntryRange[1]), TempFloats[1]);
+	REQUIRE_EQUAL(Data.FloatAttribute->GetValueFromItemKey(MetadataEntryRange[2]), TempFloats[2]);
+	REQUIRE_EQUAL(Data.FloatAttribute->GetValueFromItemKey(MetadataEntryRange[3]), TempFloats[3]);
+
+	// Points should not have a different entry key, if it was not invalid in the first place
+	REQUIRE_EQUAL(MetadataEntryRange[0], PCGMetadataEntryKey(2));
+	REQUIRE_EQUAL(MetadataEntryRange[1], PCGMetadataEntryKey(1));
+	REQUIRE_EQUAL(MetadataEntryRange[2], PCGFirstEntryKey);
+
+	// Also last point was a default entry. Making sure that we have one more entry and was set correctly
+	REQUIRE_EQUAL(Data.PointData->Metadata->GetItemCountForChild(), ValuesSize + 1);
+	REQUIRE_EQUAL(MetadataEntryRange[3], PCGMetadataEntryKey(3));
+}
+
+// Tests Get/Set with constructible type conversion (float<->double). Verifies that conversion fails without AllowConstructible and succeeds with it, for both single and range operations.
+TEST_CASE_METHOD(PCGTests::FPCGBaseTestWithBasicContext, "PCG::Accessor::Attribute::GetWithConstructibleType", "[PCG][Accessor][Attribute]")
+{
+	UPCGBasePointData* PointData = CreatePointData();
+
+	// float -> double is constructible (std::is_constructible_v<double, float>)
+	constexpr float FloatValue = 1.5f;
+	FPCGMetadataAttribute<float>* FloatAttribute = PointData->Metadata->CreateAttribute<float>(TEXT("FloatAttribute"), FloatValue, true, false);
+	TUniquePtr<IPCGAttributeAccessor> FloatAccessor = PCGAttributeAccessorHelpers::CreateAccessor(FloatAttribute, PointData->Metadata);
+	REQUIRE(FloatAccessor);
+
+	// Single point with default entry to test default value access
+	PointData->SetNumPoints(1);
+
+	SECTION("Get double from float accessor without AllowConstructible fails")
+	{
+		TUniquePtr<const IPCGAttributeAccessorKeys> Keys = PCGAttributeAccessorHelpers::CreateConstKeys(PointData, FPCGAttributePropertySelector());
+		REQUIRE(Keys);
+		double TempDouble = 0.0;
+		REQUIRE_FALSE(FloatAccessor->Get(TempDouble, *Keys));
+	}
+
+	SECTION("Get double from float accessor with AllowConstructible succeeds")
+	{
+		TUniquePtr<const IPCGAttributeAccessorKeys> Keys = PCGAttributeAccessorHelpers::CreateConstKeys(PointData, FPCGAttributePropertySelector());
+		REQUIRE(Keys);
+		double TempDouble = 0.0;
+		REQUIRE(FloatAccessor->Get(TempDouble, *Keys, EPCGAttributeAccessorFlags::AllowConstructible));
+		REQUIRE_EQUAL(TempDouble, static_cast<double>(FloatValue));
+	}
+
+	SECTION("Set double into float accessor without AllowConstructible fails")
+	{
+		TPCGValueRange<int64> MetadataEntryRange = PointData->GetMetadataEntryValueRange();
+		MetadataEntryRange[0] = PointData->Metadata->AddEntry();
+		TUniquePtr<IPCGAttributeAccessorKeys> Keys = PCGAttributeAccessorHelpers::CreateKeys(PointData, FPCGAttributePropertySelector());
+		REQUIRE(Keys);
+
+		const double DoubleValue = 2.5;
+		REQUIRE_FALSE(FloatAccessor->Set(DoubleValue, *Keys));
+	}
+
+	SECTION("Set double into float accessor with AllowConstructible succeeds")
+	{
+		TPCGValueRange<int64> MetadataEntryRange = PointData->GetMetadataEntryValueRange();
+		MetadataEntryRange[0] = PointData->Metadata->AddEntry();
+		TUniquePtr<IPCGAttributeAccessorKeys> Keys = PCGAttributeAccessorHelpers::CreateKeys(PointData, FPCGAttributePropertySelector());
+		REQUIRE(Keys);
+
+		const double DoubleValue = 2.5;
+		REQUIRE(FloatAccessor->Set(DoubleValue, *Keys, EPCGAttributeAccessorFlags::AllowConstructible));
+		REQUIRE_EQUAL(FloatAttribute->GetValueFromItemKey(MetadataEntryRange[0]), static_cast<float>(DoubleValue));
+	}
+
+	SECTION("GetRange double from float accessor with AllowConstructible succeeds")
+	{
+		constexpr int32 NumElements = 128;
+		FRandomStream RandomStream(42);
+
+		PointData->SetNumPoints(NumElements);
+		TPCGValueRange<int64> MetadataEntryRange = PointData->GetMetadataEntryValueRange();
+		for (int32 i = 0; i < NumElements; ++i)
+		{
+			MetadataEntryRange[i] = PointData->Metadata->AddEntry();
+			FloatAttribute->SetValue(MetadataEntryRange[i], RandomStream.FRand());
+		}
+
+		TUniquePtr<const IPCGAttributeAccessorKeys> Keys = PCGAttributeAccessorHelpers::CreateConstKeys(PointData, FPCGAttributePropertySelector());
+		REQUIRE(Keys);
+		REQUIRE_EQUAL(Keys->GetNum(), NumElements);
+
+		TArray<double> OutDoubles;
+		OutDoubles.SetNum(NumElements);
+
+		REQUIRE(FloatAccessor->GetRange(TArrayView<double>(OutDoubles), 0, *Keys, EPCGAttributeAccessorFlags::AllowConstructible));
+
+		for (int32 i = 0; i < NumElements; ++i)
+		{
+			REQUIRE_EQUAL(OutDoubles[i], static_cast<double>(FloatAttribute->GetValueFromItemKey(MetadataEntryRange[i])));
+		}
+	}
+
+	SECTION("SetRange double into float accessor with AllowConstructible succeeds")
+	{
+		constexpr int32 NumElements = 128;
+		FRandomStream RandomStream(42);
+
+		PointData->SetNumPoints(NumElements);
+		TPCGValueRange<int64> MetadataEntryRange = PointData->GetMetadataEntryValueRange();
+		for (int32 i = 0; i < NumElements; ++i)
+		{
+			MetadataEntryRange[i] = PointData->Metadata->AddEntry();
+		}
+
+		TArray<double> InDoubles;
+		InDoubles.SetNum(NumElements);
+		for (int32 i = 0; i < NumElements; ++i)
+		{
+			InDoubles[i] = RandomStream.FRand() * 100.0;
+		}
+
+		TUniquePtr<IPCGAttributeAccessorKeys> Keys = PCGAttributeAccessorHelpers::CreateKeys(PointData, FPCGAttributePropertySelector());
+		REQUIRE(Keys);
+		REQUIRE_EQUAL(Keys->GetNum(), NumElements);
+
+		REQUIRE(FloatAccessor->SetRange(TArrayView<const double>(InDoubles), 0, *Keys, EPCGAttributeAccessorFlags::AllowConstructible));
+
+		for (int32 i = 0; i < NumElements; ++i)
+		{
+			REQUIRE_EQUAL(FloatAttribute->GetValueFromItemKey(MetadataEntryRange[i]), static_cast<float>(InDoubles[i]));
+		}
+	}
+}
+
+// Tests Get/Set with broadcastable type conversion (float->FVector, float->FString). Verifies that conversion fails without AllowBroadcast and succeeds with it, for both single and range operations.
+TEST_CASE_METHOD(PCGTests::FPCGBaseTestWithBasicContext, "PCG::Accessor::Attribute::GetWithBroadcastableType", "[PCG][Accessor][Attribute]")
+{
+	UPCGBasePointData* PointData = CreatePointData();
+
+	// float -> FVector is broadcastable, float -> FString is broadcastable
+	constexpr float FloatValue = 3.14f;
+	FPCGMetadataAttribute<float>* FloatAttribute = PointData->Metadata->CreateAttribute<float>(TEXT("FloatAttribute"), FloatValue, true, false);
+	TUniquePtr<IPCGAttributeAccessor> FloatAccessor = PCGAttributeAccessorHelpers::CreateAccessor(FloatAttribute, PointData->Metadata);
+	REQUIRE(FloatAccessor);
+
+	// Single point with default entry to test default value access
+	PointData->SetNumPoints(1);
+
+	SECTION("Get FVector from float accessor without AllowBroadcast fails")
+	{
+		TUniquePtr<const IPCGAttributeAccessorKeys> Keys = PCGAttributeAccessorHelpers::CreateConstKeys(PointData, FPCGAttributePropertySelector());
+		REQUIRE(Keys);
+		FVector TempVector = FVector::ZeroVector;
+		REQUIRE_FALSE(FloatAccessor->Get(TempVector, *Keys));
+	}
+
+	SECTION("Get FVector from float accessor with AllowBroadcast succeeds")
+	{
+		TUniquePtr<const IPCGAttributeAccessorKeys> Keys = PCGAttributeAccessorHelpers::CreateConstKeys(PointData, FPCGAttributePropertySelector());
+		REQUIRE(Keys);
+		FVector TempVector = FVector::ZeroVector;
+		REQUIRE(FloatAccessor->Get(TempVector, *Keys, EPCGAttributeAccessorFlags::AllowBroadcast));
+		REQUIRE_EQUAL(TempVector, FVector(FloatValue));
+	}
+
+	SECTION("Get FString from float accessor with AllowBroadcast succeeds")
+	{
+		TUniquePtr<const IPCGAttributeAccessorKeys> Keys = PCGAttributeAccessorHelpers::CreateConstKeys(PointData, FPCGAttributePropertySelector());
+		REQUIRE(Keys);
+		FString TempString;
+		REQUIRE(FloatAccessor->Get(TempString, *Keys, EPCGAttributeAccessorFlags::AllowBroadcast));
+		REQUIRE_FALSE(TempString.IsEmpty());
+	}
+
+	SECTION("Set FVector into float accessor without AllowBroadcast fails")
+	{
+		TPCGValueRange<int64> MetadataEntryRange = PointData->GetMetadataEntryValueRange();
+		MetadataEntryRange[0] = PointData->Metadata->AddEntry();
+		TUniquePtr<IPCGAttributeAccessorKeys> Keys = PCGAttributeAccessorHelpers::CreateKeys(PointData, FPCGAttributePropertySelector());
+		REQUIRE(Keys);
+
+		const FVector VectorValue(5.0f, 6.0f, 7.0f);
+		REQUIRE_FALSE(FloatAccessor->Set(VectorValue, *Keys));
+	}
+
+	SECTION("Set float into FVector accessor with AllowBroadcast succeeds")
+	{
+		FPCGMetadataAttribute<FVector>* VectorAttribute = PointData->Metadata->CreateAttribute<FVector>(TEXT("VectorAttribute"), FVector::ZeroVector, false, false);
+		TUniquePtr<IPCGAttributeAccessor> VectorAccessor = PCGAttributeAccessorHelpers::CreateAccessor(VectorAttribute, PointData->Metadata);
+		REQUIRE(VectorAccessor);
+
+		TPCGValueRange<int64> MetadataEntryRange = PointData->GetMetadataEntryValueRange();
+		MetadataEntryRange[0] = PointData->Metadata->AddEntry();
+		TUniquePtr<IPCGAttributeAccessorKeys> Keys = PCGAttributeAccessorHelpers::CreateKeys(PointData, FPCGAttributePropertySelector());
+		REQUIRE(Keys);
+
+		REQUIRE(VectorAccessor->Set(FloatValue, *Keys, EPCGAttributeAccessorFlags::AllowBroadcast));
+		REQUIRE_EQUAL(VectorAttribute->GetValueFromItemKey(MetadataEntryRange[0]), FVector(FloatValue));
+	}
+
+	SECTION("GetRange FVector from float accessor with AllowBroadcast succeeds")
+	{
+		constexpr int32 NumElements = 128;
+		FRandomStream RandomStream(42);
+
+		PointData->SetNumPoints(NumElements);
+		TPCGValueRange<int64> MetadataEntryRange = PointData->GetMetadataEntryValueRange();
+		for (int32 i = 0; i < NumElements; ++i)
+		{
+			MetadataEntryRange[i] = PointData->Metadata->AddEntry();
+			FloatAttribute->SetValue(MetadataEntryRange[i], RandomStream.FRand());
+		}
+
+		TUniquePtr<const IPCGAttributeAccessorKeys> Keys = PCGAttributeAccessorHelpers::CreateConstKeys(PointData, FPCGAttributePropertySelector());
+		REQUIRE(Keys);
+		REQUIRE_EQUAL(Keys->GetNum(), NumElements);
+
+		TArray<FVector> OutVectors;
+		OutVectors.SetNum(NumElements);
+
+		REQUIRE(FloatAccessor->GetRange(TArrayView<FVector>(OutVectors), 0, *Keys, EPCGAttributeAccessorFlags::AllowBroadcast));
+
+		for (int32 i = 0; i < NumElements; ++i)
+		{
+			const float SourceFloat = FloatAttribute->GetValueFromItemKey(MetadataEntryRange[i]);
+			REQUIRE_EQUAL(OutVectors[i], FVector(SourceFloat));
+		}
+	}
+
+	SECTION("SetRange float into FVector4 accessor with AllowBroadcast succeeds")
+	{
+		constexpr int32 NumElements = 1024;
+		FRandomStream RandomStream(42);
+
+		FPCGMetadataAttribute<FVector4>* Vector4Attribute = PointData->Metadata->CreateAttribute<FVector4>(TEXT("Vector4Attribute"), FVector4::Zero(), false, false);
+		TUniquePtr<IPCGAttributeAccessor> Vector4Accessor = PCGAttributeAccessorHelpers::CreateAccessor(Vector4Attribute, PointData->Metadata);
+		REQUIRE(Vector4Accessor);
+
+		PointData->SetNumPoints(NumElements);
+		TPCGValueRange<int64> MetadataEntryRange = PointData->GetMetadataEntryValueRange();
+		for (int32 i = 0; i < NumElements; ++i)
+		{
+			MetadataEntryRange[i] = PointData->Metadata->AddEntry();
+		}
+
+		TArray<float> InFloats;
+		InFloats.SetNum(NumElements);
+		for (int32 i = 0; i < NumElements; ++i)
+		{
+			InFloats[i] = RandomStream.FRand() * 100.0f;
+		}
+
+		TUniquePtr<IPCGAttributeAccessorKeys> Keys = PCGAttributeAccessorHelpers::CreateKeys(PointData, FPCGAttributePropertySelector());
+		REQUIRE(Keys);
+		REQUIRE_EQUAL(Keys->GetNum(), NumElements);
+
+		REQUIRE(Vector4Accessor->SetRange(TArrayView<const float>(InFloats), 0, *Keys, EPCGAttributeAccessorFlags::AllowBroadcast));
+
+		for (int32 i = 0; i < NumElements; ++i)
+		{
+			REQUIRE_EQUAL(Vector4Attribute->GetValueFromItemKey(MetadataEntryRange[i]), FVector4(InFloats[i], InFloats[i], InFloats[i], InFloats[i]));
+		}
+	}
+}
+
+// Tests that FSoftObjectPath and FSoftClassPath attributes can be read as FString even with StrictType flag (legacy special case), and that non-string type access fails with StrictType.
+TEST_CASE_METHOD(PCGTests::FPCGBaseTestWithBasicContext, "PCG::Accessor::Attribute::SoftObjectPathToStringWithStrictType", "[PCG][Accessor][Attribute]")
+{
+	UPCGBasePointData* PointData = CreatePointData();
+
+	const FSoftObjectPath SoftPath(TEXT("/Game/Test/MyAsset"));
+	FPCGMetadataAttribute<FSoftObjectPath>* SoftPathAttribute = PointData->Metadata->CreateAttribute<FSoftObjectPath>(TEXT("SoftPathAttribute"), SoftPath, false, false);
+	TUniquePtr<IPCGAttributeAccessor> SoftPathAccessor = PCGAttributeAccessorHelpers::CreateAccessor(SoftPathAttribute, PointData->Metadata);
+	REQUIRE(SoftPathAccessor);
+
+	// Single point with default entry to test default value access
+	PointData->SetNumPoints(1);
+	TUniquePtr<const IPCGAttributeAccessorKeys> Keys = PCGAttributeAccessorHelpers::CreateConstKeys(PointData, FPCGAttributePropertySelector());
+	REQUIRE(Keys);
+
+	SECTION("Get FString from FSoftObjectPath with StrictType succeeds due to legacy special case")
+	{
+		FString TempString;
+		REQUIRE(SoftPathAccessor->Get(TempString, *Keys, EPCGAttributeAccessorFlags::StrictType));
+		REQUIRE_EQUAL(TempString, SoftPath.ToString());
+	}
+
+	SECTION("Get FString from FSoftObjectPath with no flags also succeeds")
+	{
+		FString TempString;
+		// Without any flags, the conversion is still invalid (no AllowBroadcast), but the special case kicks in
+		REQUIRE(SoftPathAccessor->Get(TempString, *Keys));
+		REQUIRE_EQUAL(TempString, SoftPath.ToString());
+	}
+
+	SECTION("Get FString from FSoftClassPath with StrictType succeeds")
+	{
+		const FSoftClassPath SoftClassPath(TEXT("/Game/Test/MyClass"));
+		FPCGMetadataAttribute<FSoftClassPath>* SoftClassPathAttribute = PointData->Metadata->CreateAttribute<FSoftClassPath>(TEXT("SoftClassPathAttribute"), SoftClassPath, false, false);
+		TUniquePtr<IPCGAttributeAccessor> SoftClassPathAccessor = PCGAttributeAccessorHelpers::CreateAccessor(SoftClassPathAttribute, PointData->Metadata);
+		REQUIRE(SoftClassPathAccessor);
+
+		FString TempString;
+		REQUIRE(SoftClassPathAccessor->Get(TempString, *Keys, EPCGAttributeAccessorFlags::StrictType));
+		REQUIRE_EQUAL(TempString, SoftClassPath.ToString());
+	}
+
+	SECTION("Get non-string type from FSoftObjectPath with StrictType fails")
+	{
+		float TempFloat = 0.0f;
+		REQUIRE_FALSE(SoftPathAccessor->Get(TempFloat, *Keys, EPCGAttributeAccessorFlags::StrictType));
+	}
+}
diff --git a/Engine/Plugins/PCG/Tests/Private/Metadata/PCGMetadataAttributeGenericTests.cpp b/Engine/Plugins/PCG/Tests/Private/Metadata/PCGMetadataAttributeGenericTests.cpp
--- a/Engine/Plugins/PCG/Tests/Private/Metadata/PCGMetadataAttributeGenericTests.cpp
+++ b/Engine/Plugins/PCG/Tests/Private/Metadata/PCGMetadataAttributeGenericTests.cpp
@@ -110,7 +110,6 @@ namespace PCGAttributeGenericTests
 			.ValueTypeObject = TBaseStructure<FVector>::Get()
 		};
 		
-		AttributeTester.ExpectedDesc = StructDesc;
 		AttributeTester.OverrideDesc = StructDesc;
 	
 		return AttributeTester;
