[PCG] Factorize testers for generic attributes to re-use them in other unit tests

* Also fix missing pragma once in one of the helpers

#rnx
#Synced-CL 52060281

diff --git a/Engine/Plugins/PCG/Tests/Private/Metadata/PCGMetadataAttributeGenericTests.cpp b/Engine/Plugins/PCG/Tests/Private/Metadata/PCGMetadataAttributeGenericTests.cpp
--- a/Engine/Plugins/PCG/Tests/Private/Metadata/PCGMetadataAttributeGenericTests.cpp
+++ b/Engine/Plugins/PCG/Tests/Private/Metadata/PCGMetadataAttributeGenericTests.cpp
@@ -1,6 +1,7 @@
 // Copyright Epic Games, Inc. All Rights Reserved.
 
 #include "PCGTestsCommon.h"
+#include "Metadata/PCGMetadataAttributeTestsCommonHelper.h"
 
 #include "PCGParamData.h"
 #include "Data/PCGPointArrayData.h"
@@ -16,350 +17,6 @@
 #include "Async/Async.h"
 #include "HAL/Event.h"
 
-namespace PCGAttributeGenericTests
-{
-	/**
-	 * Base class for a tester that encapsulate everything that is type specific.
-	 */
-	template <typename T, typename InReadType = T>
-	struct TypedAttributeTester
-	{
-		using Type = T;
-		using ReadType = InReadType;
-
-		TypedAttributeTester(TFunction<void(FRandomStream&, T&)> InGenerateRandom, 
-			TFunction<bool(const T&, const ReadType&)> InVerify, 
-			T InDefaultValue,
-			FPCGMetadataAttributeDesc InExpectedDesc)
-			: GenerateRandom(MoveTemp(InGenerateRandom))
-			, Verify(MoveTemp(InVerify))
-			, DefaultValue(InDefaultValue)
-			, ExpectedDesc(InExpectedDesc)
-		{}
-
-		TFunction<void(FRandomStream&, T&)> GenerateRandom;
-		TFunction<bool(const T&, const ReadType&)> Verify;
-		T DefaultValue;
-		FPCGMetadataAttributeDesc ExpectedDesc;
-		TOptional<FPCGMetadataAttributeDesc> OverrideDesc;
-	};
-	
-	static const FName AttributeName = "Attr";
-	static const FName AttributeName2 = "Attr2";
-	static const FName AttributeName3 = "Attr3";
-	
-	/** Getter to create a tester for different types. */
-	TypedAttributeTester<float> FloatTester()
-	{
-		auto GenerateRandom = [](FRandomStream& RandomStream, float& OutValue) { OutValue = RandomStream.FRand(); };
-		
-		auto Verify = [](const float& LHS, const float& RHS) { return LHS == RHS; };
-		
-		float DefaultValue = 5.0f;
-		
-		FPCGMetadataAttributeDesc ExpectedDesc
-		{
-			.Name = AttributeName,
-			.ValueType = EPCGMetadataTypes::Float
-		};
-		
-		return TypedAttributeTester<float>(GenerateRandom, Verify, DefaultValue, MoveTemp(ExpectedDesc));
-	}
-	
-	TypedAttributeTester<double> DoubleTester()
-	{
-		auto GenerateRandom = [](FRandomStream& RandomStream, double& OutValue) { OutValue = RandomStream.FRand(); };
-		
-		auto Verify = [](const double& LHS, const double& RHS) { return LHS == RHS; };
-		
-		double DefaultValue = 5.0;
-		
-		FPCGMetadataAttributeDesc ExpectedDesc
-		{
-			.Name = AttributeName,
-			.ValueType = EPCGMetadataTypes::Double
-		};
-		
-		return TypedAttributeTester<double>(GenerateRandom, Verify, DefaultValue, MoveTemp(ExpectedDesc));
-	}
-	
-	TypedAttributeTester<FVector> VectorTester_Basic()
-	{
-		auto GenerateRandom = [](FRandomStream& RandomStream, FVector& OutValue) { OutValue = RandomStream.VRand(); };
-		
-		auto Verify = [](const FVector& LHS, const FVector& RHS) { return LHS == RHS; };
-		
-		FVector DefaultValue = FVector(1.0, 2.0, 3.0);
-		
-		FPCGMetadataAttributeDesc ExpectedDesc
-		{
-			.Name = AttributeName,
-			.ValueType = EPCGMetadataTypes::Vector
-		};
-			
-		return TypedAttributeTester<FVector>{GenerateRandom, Verify, DefaultValue, MoveTemp(ExpectedDesc)};
-	}
-		
-	TypedAttributeTester<FVector> VectorTester_Struct()
-	{
-		TypedAttributeTester<FVector> AttributeTester = VectorTester_Basic();
-		FPCGMetadataAttributeDesc StructDesc
-		{
-			.Name = AttributeName,
-			.ValueType = EPCGMetadataTypes::Struct,
-			.ValueTypeObject = TBaseStructure<FVector>::Get()
-		};
-		
-		AttributeTester.OverrideDesc = StructDesc;
-	
-		return AttributeTester;
-	}
-	
-	TypedAttributeTester<FString, const FString*> StringTester()
-	{
-		auto GenerateRandom = [](FRandomStream& RandomStream, FString& OutValue)
-			{
-				static int i = 0; 
-				OutValue = FString::Printf(TEXT("This is a string with some characters and a value: %d"), i++ % 25);
-			};
-		
-		auto Verify = [](const FString& LHS, FString const* const& RHS) { return LHS == *RHS; };
-		
-		FString DefaultValue = TEXT("Hi");
-		
-		FPCGMetadataAttributeDesc ExpectedDesc
-		{
-			.Name = AttributeName,
-			.ValueType = EPCGMetadataTypes::String
-		};
-		
-		return TypedAttributeTester<FString, const FString*>(GenerateRandom, Verify, MoveTemp(DefaultValue), MoveTemp(ExpectedDesc));
-	}
-	
-	TypedAttributeTester<FPCGPoint, const FPCGPoint*> FPCGPointTester()
-	{		
-		auto GenerateRandom = [](FRandomStream& RandomStream, FPCGPoint& OutValue)
-			{
-				OutValue.Density = RandomStream.FRand();
-				OutValue.Transform.SetLocation(RandomStream.VRand());
-			};
-		
-		auto Verify = [](const FPCGPoint& LHS, FPCGPoint const* const& RHS)
-			{
-				return LHS.Density == RHS->Density && LHS.Transform.GetLocation() == RHS->Transform.GetLocation();
-			};
-		
-		FPCGPoint DefaultValue{FTransform{5.0 * FVector::OneVector}, 0.4f, 42};
-		
-		FPCGMetadataAttributeDesc ExpectedDesc
-		{
-			.Name = AttributeName,
-			.ValueType = EPCGMetadataTypes::Struct,
-			.ValueTypeObject = FPCGPoint::StaticStruct()
-		};
-		
-		return TypedAttributeTester<FPCGPoint, const FPCGPoint*>(GenerateRandom, Verify, MoveTemp(DefaultValue), MoveTemp(ExpectedDesc));
-	}
-	
-	TypedAttributeTester<TArray<double>, TConstArrayView<double>> ArrayDoubleTester()
-	{		
-		auto GenerateRandom = [](FRandomStream& RandomStream, TArray<double>& OutValue)
-		{
-			const int32 RandomCount = RandomStream.RandRange(1, 4);
-			for (int j = 0; j < RandomCount; ++j)
-			{
-				OutValue.Add(RandomStream.FRand());
-			}
-		};
-		
-		auto Verify = [](const TArray<double>& LHS, const TConstArrayView<double>& RHS)
-			{
-				if (LHS.Num() != RHS.Num())
-				{
-					return false;
-				}
-				
-				for (int i = 0; i < LHS.Num(); ++i)
-				{
-					if (LHS[i] != RHS[i])
-					{
-						return false;
-					}
-				}
-				
-				return true;
-			};
-		
-		TArray<double> DefaultValue = {1, 2, 3, 4, 5};
-		
-		FPCGMetadataAttributeDesc ExpectedDesc
-		{
-			.Name = AttributeName,
-			.ValueType = EPCGMetadataTypes::Double,
-			.ContainerTypes = {EPCGMetadataAttributeContainerTypes::Array}
-		};
-		
-		return TypedAttributeTester<TArray<double>, TConstArrayView<double>>(GenerateRandom, Verify, MoveTemp(DefaultValue), MoveTemp(ExpectedDesc));
-	}
-	
-	TypedAttributeTester<TArray<FPCGPoint>, TConstArrayView<FPCGPoint>> ArrayFPCGPointTester()
-	{	
-		auto GenerateRandom = [](FRandomStream& RandomStream, TArray<FPCGPoint>& OutValue)
-			{
-				const int32 RandomCount = RandomStream.RandRange(1, 4);
-				for (int j = 0; j < RandomCount; ++j)
-				{
-					FPCGPoint& Value = OutValue.Emplace_GetRef();
-					Value.Density = RandomStream.FRand();
-					Value.Transform.SetLocation(RandomStream.VRand());
-				}
-			};
-		
-		auto Verify = [](const TArray<FPCGPoint>& LHS, const TConstArrayView<FPCGPoint>& RHS)
-			{
-				if (LHS.Num() != RHS.Num())
-				{
-					return false;
-				}
-				
-				for (int i = 0; i < LHS.Num(); ++i)
-				{
-					if (LHS[i].Density != RHS[i].Density || LHS[i].Transform.GetLocation() != RHS[i].Transform.GetLocation())
-					{
-						return false;
-					}
-				}
-				
-				return true;
-			};
-		
-		TArray<FPCGPoint> DefaultValue = {
-			{FTransform{5.0 * FVector::OneVector}, 0.4f, 42},
-			{FTransform{6.0 * FVector::OneVector}, 0.5f, 43},
-			{FTransform{7.0 * FVector::OneVector}, 0.6f, 44},
-		};
-		
-		FPCGMetadataAttributeDesc ExpectedDesc
-		{
-			.Name = AttributeName,
-			.ValueType = EPCGMetadataTypes::Struct,
-			.ContainerTypes = {EPCGMetadataAttributeContainerTypes::Array},
-			.ValueTypeObject = FPCGPoint::StaticStruct()
-		};
-		
-		return TypedAttributeTester<TArray<FPCGPoint>, TConstArrayView<FPCGPoint>>(GenerateRandom, Verify, MoveTemp(DefaultValue), MoveTemp(ExpectedDesc));
-	}
-	
-	TypedAttributeTester<TSet<double>, PCG::TScriptSetWrapper<double>> SetDoubleTester()
-	{		
-		auto GenerateRandom = [](FRandomStream& RandomStream, TSet<double>& OutValue)
-			{
-				const int32 RandomCount = RandomStream.RandRange(1, 4);
-				for (int j = 0; j < RandomCount; ++j)
-				{
-					OutValue.Add(RandomStream.FRand());
-				}
-			};
-		
-		auto Verify = [](const TSet<double>& LHS, const PCG::TScriptSetWrapper<double>& RHS)
-			{
-				if (LHS.Num() != RHS.Num())
-				{
-					return false;
-				}
-				
-				for (double Value : LHS)
-				{
-					if (!RHS.Contains(Value))
-					{
-						return false;
-					}
-				}
-				
-				return true;
-			};
-		
-		TSet<double> DefaultValue = {1, 2, 3, 4, 5};
-		
-		FPCGMetadataAttributeDesc ExpectedDesc
-		{
-			.Name = AttributeName,
-			.ValueType = EPCGMetadataTypes::Double,
-			.ContainerTypes = {EPCGMetadataAttributeContainerTypes::Set}
-		};
-		
-		return TypedAttributeTester<TSet<double>, PCG::TScriptSetWrapper<double>>(GenerateRandom, Verify, MoveTemp(DefaultValue), MoveTemp(ExpectedDesc));
-	}
-	
-	TypedAttributeTester<TMap<FString, double>, PCG::TScriptMapWrapper<FString, double>> MapStringDoubleTester()
-	{
-		auto GenerateRandom = [](FRandomStream& RandomStream, TMap<FString, double>& OutValue)
-			{
-				const int32 RandomCount = RandomStream.RandRange(1, 4);
-				for (int j = 0; j < RandomCount; ++j)
-				{
-					OutValue.Emplace(LexToString(j), RandomStream.FRand());
-				}
-			};
-		
-		auto Verify = [](const TMap<FString, double>& LHS, const PCG::TScriptMapWrapper<FString, double>& RHS)
-			{
-				if (LHS.Num() != RHS.Num())
-				{
-					return false;
-				}
-				
-				for (const auto& [Key, Value] : LHS)
-				{
-					const double* OtherValue = RHS.Find(Key);
-					if (!OtherValue || Value != *OtherValue)
-					{
-						return false;
-					}
-				}
-				
-				return true;
-			};
-		
-		TMap<FString, double> DefaultValue =
-		{
-			{TEXT("Hi"), 2.0},
-			{TEXT("Ho"), 3.0}
-		};
-		
-		FPCGMetadataAttributeDesc ExpectedDesc
-		{
-			.Name = AttributeName,
-			.ValueType = EPCGMetadataTypes::Double,
-			.ContainerTypes = {EPCGMetadataAttributeContainerTypes::Map},
-			.KeyType = EPCGMetadataTypes::String
-		};
-		
-		return TypedAttributeTester<TMap<FString, double>, PCG::TScriptMapWrapper<FString, double>>(GenerateRandom, Verify, MoveTemp(DefaultValue), MoveTemp(ExpectedDesc));
-	}
-	
-	TypedAttributeTester<EPCGMetadataTypes> EnumTester()
-	{
-		auto GenerateRandom = [](FRandomStream& RandomStream, EPCGMetadataTypes& OutValue)
-			{
-				OutValue = static_cast<EPCGMetadataTypes>(RandomStream.RandRange(0, static_cast<int32>(EPCGMetadataTypes::Count)));
-			};
-		
-		auto Verify = [](const EPCGMetadataTypes& LHS, const EPCGMetadataTypes& RHS) { return LHS == RHS; };
-		
-		EPCGMetadataTypes DefaultValue = EPCGMetadataTypes::Unknown;
-		
-		FPCGMetadataAttributeDesc ExpectedDesc
-		{
-			.Name = AttributeName,
-			.ValueType = EPCGMetadataTypes::Enum,
-			.ValueTypeObject = StaticEnum<EPCGMetadataTypes>()
-		};
-		
-		return TypedAttributeTester<EPCGMetadataTypes>(GenerateRandom, Verify, DefaultValue, MoveTemp(ExpectedDesc));
-	}
-}
-
 /**
  * Test suite to verify basic functionality with generic attributes
  * Rely on the TypedAttributeTester that handles anything that is type specific, the rest is templated in
@@ -380,7 +37,7 @@ TEST_CASE_METHOD(PCGTests::FPCGBaseTestWithBasicContext, "PCG::AttributeGeneric:
 {
 	constexpr int32 Seed = 42;
 	
-	auto Tester = [this]<typename T, typename ReadType>(PCGAttributeGenericTests::TypedAttributeTester<T, ReadType> AttributeTester)
+	auto Tester = [this]<typename T, typename ReadType>(PCGAttributeTestsCommonHelper::TypedAttributeTester<T, ReadType> AttributeTester)
 	{
 		// Preparing the data
 		constexpr int32 NumPoints = 100;
@@ -405,7 +62,7 @@ TEST_CASE_METHOD(PCGTests::FPCGBaseTestWithBasicContext, "PCG::AttributeGeneric:
 		}
 		else
 		{
-			GenericAttribute = Domain->CreateGenericAttribute<T>(PCGAttributeGenericTests::AttributeName, AttributeTester.DefaultValue);
+			GenericAttribute = Domain->CreateGenericAttribute<T>(PCGAttributeTestsCommonHelper::AttributeName, AttributeTester.DefaultValue);
 			REQUIRE_NOT_EQUAL(GenericAttribute, nullptr);
 		}
 		
@@ -448,7 +105,7 @@ TEST_CASE_METHOD(PCGTests::FPCGBaseTestWithBasicContext, "PCG::AttributeGeneric:
 
 			TPCGValueRange<int64> ChildPointArrayDataMetadataEntryRange = ChildPointArrayData->GetMetadataEntryValueRange(/*bAllocate=*/false);
 
-			FPCGMetadataAttributeGeneric* ChildGenericAttribute = ChildDomain->GetMutableGenericAttribute(PCGAttributeGenericTests::AttributeName);
+			FPCGMetadataAttributeGeneric* ChildGenericAttribute = ChildDomain->GetMutableGenericAttribute(PCGAttributeTestsCommonHelper::AttributeName);
 			REQUIRE_NOT_EQUAL(ChildGenericAttribute, nullptr);
 			
 			SECTION("Expected attribute descriptor")
@@ -492,60 +149,60 @@ TEST_CASE_METHOD(PCGTests::FPCGBaseTestWithBasicContext, "PCG::AttributeGeneric:
 	
 	SECTION("Float Attribute")
 	{
-		Tester(PCGAttributeGenericTests::FloatTester());
+		Tester(PCGAttributeTestsCommonHelper::FloatTester());
 	}
 	
 	SECTION("Double Attribute")
 	{
-		Tester(PCGAttributeGenericTests::DoubleTester());
+		Tester(PCGAttributeTestsCommonHelper::DoubleTester());
 	}
 	
 	SECTION("Vector Attribute")
 	{
 		SECTION("As Basic Type")
 		{
-			Tester(PCGAttributeGenericTests::VectorTester_Basic());
+			Tester(PCGAttributeTestsCommonHelper::VectorTester_Basic());
 		}
 		
 		SECTION("As Struct")
 		{
-			Tester(PCGAttributeGenericTests::VectorTester_Struct());
+			Tester(PCGAttributeTestsCommonHelper::VectorTester_Struct());
 		}
 	}
 	
 	SECTION("String Attribute")
 	{
-		Tester(PCGAttributeGenericTests::StringTester());
+		Tester(PCGAttributeTestsCommonHelper::StringTester());
 	}
 	
 	SECTION("FPCGPoint Attribute")
 	{		
-		Tester(PCGAttributeGenericTests::FPCGPointTester());
+		Tester(PCGAttributeTestsCommonHelper::FPCGPointTester());
 	}
 	
 	SECTION("Array of Double Attribute")
 	{		
-		Tester(PCGAttributeGenericTests::ArrayDoubleTester());
+		Tester(PCGAttributeTestsCommonHelper::ArrayDoubleTester());
 	}
 	
 	SECTION("Array of FPCGPoint Attribute")
 	{	
-		Tester(PCGAttributeGenericTests::ArrayFPCGPointTester());
+		Tester(PCGAttributeTestsCommonHelper::ArrayFPCGPointTester());
 	}
 	
 	SECTION("Set of Double Attribute")
 	{		
-		Tester(PCGAttributeGenericTests::SetDoubleTester());
+		Tester(PCGAttributeTestsCommonHelper::SetDoubleTester());
 	}
 	
 	SECTION("Map of String -> Double Attribute")
 	{
-		Tester(PCGAttributeGenericTests::MapStringDoubleTester());
+		Tester(PCGAttributeTestsCommonHelper::MapStringDoubleTester());
 	}
 	
 	SECTION("Enum attribute")
 	{
-		Tester(PCGAttributeGenericTests::EnumTester());
+		Tester(PCGAttributeTestsCommonHelper::EnumTester());
 	}
 }
 
@@ -556,7 +213,7 @@ TEST_CASE_METHOD(PCGTests::FPCGBaseTestWithBasicContext, "PCG::AttributeGeneric:
 {
 	constexpr int32 Seed = 42;
 	
-	auto Tester = [this]<typename T, typename ReadType>(PCGAttributeGenericTests::TypedAttributeTester<T, ReadType> AttributeTester)
+	auto Tester = [this]<typename T, typename ReadType>(PCGAttributeTestsCommonHelper::TypedAttributeTester<T, ReadType> AttributeTester)
 	{
 		// Preparing the data
 		constexpr int32 NumEntries = 10;
@@ -575,7 +232,7 @@ TEST_CASE_METHOD(PCGTests::FPCGBaseTestWithBasicContext, "PCG::AttributeGeneric:
 		}
 		else
 		{
-			GenericAttribute = Domain->CreateGenericAttribute<T>(PCGAttributeGenericTests::AttributeName, AttributeTester.DefaultValue);
+			GenericAttribute = Domain->CreateGenericAttribute<T>(PCGAttributeTestsCommonHelper::AttributeName, AttributeTester.DefaultValue);
 			REQUIRE_NOT_EQUAL(GenericAttribute, nullptr);
 		}
 		
@@ -605,7 +262,7 @@ TEST_CASE_METHOD(PCGTests::FPCGBaseTestWithBasicContext, "PCG::AttributeGeneric:
 
 			FPCGMetadataDomain* ChildDomain = ChildParamData->MutableMetadata()->GetDefaultMetadataDomain();
 
-			FPCGMetadataAttributeGeneric* ChildGenericAttribute = ChildDomain->GetMutableGenericAttribute(PCGAttributeGenericTests::AttributeName);
+			FPCGMetadataAttributeGeneric* ChildGenericAttribute = ChildDomain->GetMutableGenericAttribute(PCGAttributeTestsCommonHelper::AttributeName);
 			REQUIRE_NOT_EQUAL(ChildGenericAttribute, nullptr);
 			
 			SECTION("Expected attribute descriptor")
@@ -649,60 +306,60 @@ TEST_CASE_METHOD(PCGTests::FPCGBaseTestWithBasicContext, "PCG::AttributeGeneric:
 	
 	SECTION("Float Attribute")
 	{
-		Tester(PCGAttributeGenericTests::FloatTester());
+		Tester(PCGAttributeTestsCommonHelper::FloatTester());
 	}
 	
 	SECTION("Double Attribute")
 	{
-		Tester(PCGAttributeGenericTests::DoubleTester());
+		Tester(PCGAttributeTestsCommonHelper::DoubleTester());
 	}
 	
 	SECTION("Vector Attribute")
 	{
 		SECTION("As Basic Type")
 		{
-			Tester(PCGAttributeGenericTests::VectorTester_Basic());
+			Tester(PCGAttributeTestsCommonHelper::VectorTester_Basic());
 		}
 		
 		SECTION("As Struct")
 		{
-			Tester(PCGAttributeGenericTests::VectorTester_Struct());
+			Tester(PCGAttributeTestsCommonHelper::VectorTester_Struct());
 		}
 	}
 	
 	SECTION("String Attribute")
 	{
-		Tester(PCGAttributeGenericTests::StringTester());
+		Tester(PCGAttributeTestsCommonHelper::StringTester());
 	}
 	
 	SECTION("FPCGPoint Attribute")
 	{		
-		Tester(PCGAttributeGenericTests::FPCGPointTester());
+		Tester(PCGAttributeTestsCommonHelper::FPCGPointTester());
 	}
 	
 	SECTION("Array of Double Attribute")
 	{		
-		Tester(PCGAttributeGenericTests::ArrayDoubleTester());
+		Tester(PCGAttributeTestsCommonHelper::ArrayDoubleTester());
 	}
 	
 	SECTION("Array of FPCGPoint Attribute")
 	{	
-		Tester(PCGAttributeGenericTests::ArrayFPCGPointTester());
+		Tester(PCGAttributeTestsCommonHelper::ArrayFPCGPointTester());
 	}
 	
 	SECTION("Set of Double Attribute")
 	{		
-		Tester(PCGAttributeGenericTests::SetDoubleTester());
+		Tester(PCGAttributeTestsCommonHelper::SetDoubleTester());
 	}
 	
 	SECTION("Map of String -> Double Attribute")
 	{
-		Tester(PCGAttributeGenericTests::MapStringDoubleTester());
+		Tester(PCGAttributeTestsCommonHelper::MapStringDoubleTester());
 	}
 	
 	SECTION("Enum attribute")
 	{
-		Tester(PCGAttributeGenericTests::EnumTester());
+		Tester(PCGAttributeTestsCommonHelper::EnumTester());
 	}
 }
 
@@ -716,7 +373,7 @@ TEST_CASE_METHOD(PCGTests::FPCGBaseTestWithBasicContext, "PCG::AttributeGeneric:
 		UPCGParamData* ParamData = CreateData<UPCGParamData>();
 		FPCGMetadataDomain* Domain = ParamData->MutableMetadata()->GetDefaultMetadataDomain();
 		
-		FPCGMetadataAttributeGeneric* GenericAttribute = Domain->CreateGenericAttribute<RealType>(PCGAttributeGenericTests::AttributeName, InValue);
+		FPCGMetadataAttributeGeneric* GenericAttribute = Domain->CreateGenericAttribute<RealType>(PCGAttributeTestsCommonHelper::AttributeName, InValue);
 		
 		SECTION("Get Value with wrong type")
 		{
@@ -866,9 +523,9 @@ TEST_CASE_METHOD(PCGTests::FPCGBaseTestWithBasicContext, "PCG::AttributeGeneric:
 	};
 	
 	// Initialize the default value to a value we will add later on.
-	FPCGMetadataAttributeGeneric* GenericAttribute = Domain->CreateGenericAttribute<FString>(PCGAttributeGenericTests::AttributeName, TEXT("Value: 0"));
-	FPCGMetadataAttributeGeneric* GenericAttributePath = Domain->CreateGenericAttribute<FSoftObjectPath>(PCGAttributeGenericTests::AttributeName2, FSoftObjectPath(FString::Format(*BasePathFormat, { 0 })));
-	FPCGMetadataAttributeGeneric* GenericAttributeClass = Domain->CreateGenericAttribute<FSoftClassPath>(PCGAttributeGenericTests::AttributeName3, UClassesPool[0]);
+	FPCGMetadataAttributeGeneric* GenericAttribute = Domain->CreateGenericAttribute<FString>(PCGAttributeTestsCommonHelper::AttributeName, TEXT("Value: 0"));
+	FPCGMetadataAttributeGeneric* GenericAttributePath = Domain->CreateGenericAttribute<FSoftObjectPath>(PCGAttributeTestsCommonHelper::AttributeName2, FSoftObjectPath(FString::Format(*BasePathFormat, { 0 })));
+	FPCGMetadataAttributeGeneric* GenericAttributeClass = Domain->CreateGenericAttribute<FSoftClassPath>(PCGAttributeTestsCommonHelper::AttributeName3, UClassesPool[0]);
 	REQUIRE_NOT_EQUAL(GenericAttribute, nullptr);
 	REQUIRE_NOT_EQUAL(GenericAttributePath, nullptr);
 	REQUIRE_NOT_EQUAL(GenericAttributeClass, nullptr);
@@ -945,7 +602,7 @@ TEST_CASE_METHOD(PCGTests::FPCGBaseTestWithBasicContext, "PCG::AttributeGeneric:
 
 TEST_CASE_METHOD(PCGTests::FPCGBaseTestWithBasicContext, "PCG::AttributeGeneric::MultiThreadedWrite", "[PCG][AttributeGeneric]")
 {
-	auto Tester = [this]<typename T, typename ReadType>(PCGAttributeGenericTests::TypedAttributeTester<T, ReadType> AttributeTester, bool bLockless)
+	auto Tester = [this]<typename T, typename ReadType>(PCGAttributeTestsCommonHelper::TypedAttributeTester<T, ReadType> AttributeTester, bool bLockless)
 	{
 		// Preparing the data
 		constexpr int32 NumEntries = 10000;
@@ -960,7 +617,7 @@ TEST_CASE_METHOD(PCGTests::FPCGBaseTestWithBasicContext, "PCG::AttributeGeneric:
 	
 		FRandomStream RandomStream(Seed);
 	
-		FPCGMetadataAttributeGeneric* GenericAttribute = Domain->CreateGenericAttribute<T>(PCGAttributeGenericTests::AttributeName, AttributeTester.DefaultValue);
+		FPCGMetadataAttributeGeneric* GenericAttribute = Domain->CreateGenericAttribute<T>(PCGAttributeTestsCommonHelper::AttributeName, AttributeTester.DefaultValue);
 		REQUIRE_NOT_EQUAL(GenericAttribute, nullptr);
 	
 		TStaticArray<T, NumEntries> Values{};
@@ -1054,12 +711,12 @@ TEST_CASE_METHOD(PCGTests::FPCGBaseTestWithBasicContext, "PCG::AttributeGeneric:
 	{
 		SECTION("With lock")
 		{
-			Tester(PCGAttributeGenericTests::FloatTester(), /*bLockless=*/false);
+			Tester(PCGAttributeTestsCommonHelper::FloatTester(), /*bLockless=*/false);
 		}
 		
 		SECTION("Without lock")
 		{
-			Tester(PCGAttributeGenericTests::FloatTester(), /*bLockless=*/true);
+			Tester(PCGAttributeTestsCommonHelper::FloatTester(), /*bLockless=*/true);
 		}
 	}
 	
@@ -1067,12 +724,12 @@ TEST_CASE_METHOD(PCGTests::FPCGBaseTestWithBasicContext, "PCG::AttributeGeneric:
 	{
 		SECTION("With lock")
 		{
-			Tester(PCGAttributeGenericTests::DoubleTester(), /*bLockless=*/false);
+			Tester(PCGAttributeTestsCommonHelper::DoubleTester(), /*bLockless=*/false);
 		}
 		
 		SECTION("Without lock")
 		{
-			Tester(PCGAttributeGenericTests::DoubleTester(), /*bLockless=*/true);
+			Tester(PCGAttributeTestsCommonHelper::DoubleTester(), /*bLockless=*/true);
 		}
 	}
 	
@@ -1080,12 +737,12 @@ TEST_CASE_METHOD(PCGTests::FPCGBaseTestWithBasicContext, "PCG::AttributeGeneric:
 	{
 		SECTION("With lock")
 		{
-			Tester(PCGAttributeGenericTests::VectorTester_Basic(), /*bLockless=*/false);
+			Tester(PCGAttributeTestsCommonHelper::VectorTester_Basic(), /*bLockless=*/false);
 		}
 		
 		SECTION("Without lock")
 		{
-			Tester(PCGAttributeGenericTests::VectorTester_Basic(), /*bLockless=*/true);
+			Tester(PCGAttributeTestsCommonHelper::VectorTester_Basic(), /*bLockless=*/true);
 		}
 	}
 	
@@ -1093,12 +750,12 @@ TEST_CASE_METHOD(PCGTests::FPCGBaseTestWithBasicContext, "PCG::AttributeGeneric:
 	{
 		SECTION("With lock")
 		{
-			Tester(PCGAttributeGenericTests::StringTester(), /*bLockless=*/false);
+			Tester(PCGAttributeTestsCommonHelper::StringTester(), /*bLockless=*/false);
 		}
 		
 		SECTION("Without lock")
 		{
-			Tester(PCGAttributeGenericTests::StringTester(), /*bLockless=*/true);
+			Tester(PCGAttributeTestsCommonHelper::StringTester(), /*bLockless=*/true);
 		}
 	}
 	
@@ -1106,12 +763,12 @@ TEST_CASE_METHOD(PCGTests::FPCGBaseTestWithBasicContext, "PCG::AttributeGeneric:
 	{		
 		SECTION("With lock")
 		{
-			Tester(PCGAttributeGenericTests::FPCGPointTester(), /*bLockless=*/false);
+			Tester(PCGAttributeTestsCommonHelper::FPCGPointTester(), /*bLockless=*/false);
 		}
 		
 		SECTION("Without lock")
 		{
-			Tester(PCGAttributeGenericTests::FPCGPointTester(), /*bLockless=*/true);
+			Tester(PCGAttributeTestsCommonHelper::FPCGPointTester(), /*bLockless=*/true);
 		}
 	}
 	
@@ -1119,12 +776,12 @@ TEST_CASE_METHOD(PCGTests::FPCGBaseTestWithBasicContext, "PCG::AttributeGeneric:
 	{		
 		SECTION("With lock")
 		{
-			Tester(PCGAttributeGenericTests::ArrayDoubleTester(), /*bLockless=*/false);
+			Tester(PCGAttributeTestsCommonHelper::ArrayDoubleTester(), /*bLockless=*/false);
 		}
 		
 		SECTION("Without lock")
 		{
-			Tester(PCGAttributeGenericTests::ArrayDoubleTester(), /*bLockless=*/true);
+			Tester(PCGAttributeTestsCommonHelper::ArrayDoubleTester(), /*bLockless=*/true);
 		}
 	}
 	
@@ -1132,12 +789,12 @@ TEST_CASE_METHOD(PCGTests::FPCGBaseTestWithBasicContext, "PCG::AttributeGeneric:
 	{	
 		SECTION("With lock")
 		{
-			Tester(PCGAttributeGenericTests::ArrayFPCGPointTester(), /*bLockless=*/false);
+			Tester(PCGAttributeTestsCommonHelper::ArrayFPCGPointTester(), /*bLockless=*/false);
 		}
 		
 		SECTION("Without lock")
 		{
-			Tester(PCGAttributeGenericTests::ArrayFPCGPointTester(), /*bLockless=*/true);
+			Tester(PCGAttributeTestsCommonHelper::ArrayFPCGPointTester(), /*bLockless=*/true);
 		}
 	}
 	
@@ -1145,12 +802,12 @@ TEST_CASE_METHOD(PCGTests::FPCGBaseTestWithBasicContext, "PCG::AttributeGeneric:
 	{		
 		SECTION("With lock")
 		{
-			Tester(PCGAttributeGenericTests::SetDoubleTester(), /*bLockless=*/false);
+			Tester(PCGAttributeTestsCommonHelper::SetDoubleTester(), /*bLockless=*/false);
 		}
 		
 		SECTION("Without lock")
 		{
-			Tester(PCGAttributeGenericTests::SetDoubleTester(), /*bLockless=*/true);
+			Tester(PCGAttributeTestsCommonHelper::SetDoubleTester(), /*bLockless=*/true);
 		}
 	}
 	
@@ -1158,12 +815,12 @@ TEST_CASE_METHOD(PCGTests::FPCGBaseTestWithBasicContext, "PCG::AttributeGeneric:
 	{
 		SECTION("With lock")
 		{
-			Tester(PCGAttributeGenericTests::MapStringDoubleTester(), /*bLockless=*/false);
+			Tester(PCGAttributeTestsCommonHelper::MapStringDoubleTester(), /*bLockless=*/false);
 		}
 		
 		SECTION("Without lock")
 		{
-			Tester(PCGAttributeGenericTests::MapStringDoubleTester(), /*bLockless=*/true);
+			Tester(PCGAttributeTestsCommonHelper::MapStringDoubleTester(), /*bLockless=*/true);
 		}
 	}
 	
@@ -1171,12 +828,12 @@ TEST_CASE_METHOD(PCGTests::FPCGBaseTestWithBasicContext, "PCG::AttributeGeneric:
 	{
 		SECTION("With lock")
 		{
-			Tester(PCGAttributeGenericTests::EnumTester(), /*bLockless=*/false);
+			Tester(PCGAttributeTestsCommonHelper::EnumTester(), /*bLockless=*/false);
 		}
 		
 		SECTION("Without lock")
 		{
-			Tester(PCGAttributeGenericTests::EnumTester(), /*bLockless=*/true);
+			Tester(PCGAttributeTestsCommonHelper::EnumTester(), /*bLockless=*/true);
 		}
 	}
 }
@@ -1189,7 +846,7 @@ TEST_CASE_METHOD(PCGTests::FPCGBaseTestWithBasicContext, "PCG::AttributeGeneric:
  */
 TEST_CASE_METHOD(PCGTests::FPCGBaseTestWithBasicContext, "PCG::AttributeGeneric::SetValueSelfReference", "[PCG][AttributeGeneric]")
 {
-	auto Tester = [this]<typename T, typename ReadType>(PCGAttributeGenericTests::TypedAttributeTester<T, ReadType> AttributeTester)
+	auto Tester = [this]<typename T, typename ReadType>(PCGAttributeTestsCommonHelper::TypedAttributeTester<T, ReadType> AttributeTester)
 	{
 		constexpr int32 NumInitialEntries = 10;
 		// Enough self-copies to trigger multiple reallocations of the underlying FScriptArray.
@@ -1200,7 +857,7 @@ TEST_CASE_METHOD(PCGTests::FPCGBaseTestWithBasicContext, "PCG::AttributeGeneric:
 
 		FRandomStream RandomStream(42);
 
-		FPCGMetadataAttributeGeneric* GenericAttribute = Domain->CreateGenericAttribute<T>(PCGAttributeGenericTests::AttributeName, AttributeTester.DefaultValue);
+		FPCGMetadataAttributeGeneric* GenericAttribute = Domain->CreateGenericAttribute<T>(PCGAttributeTestsCommonHelper::AttributeName, AttributeTester.DefaultValue);
 		REQUIRE_NOT_EQUAL(GenericAttribute, nullptr);
 
 		// Populate initial entries with random values.
@@ -1236,22 +893,22 @@ TEST_CASE_METHOD(PCGTests::FPCGBaseTestWithBasicContext, "PCG::AttributeGeneric:
 
 	SECTION("Double (POD type)")
 	{
-		Tester(PCGAttributeGenericTests::DoubleTester());
+		Tester(PCGAttributeTestsCommonHelper::DoubleTester());
 	}
 
 	SECTION("String (non-POD, compressed)")
 	{
-		Tester(PCGAttributeGenericTests::StringTester());
+		Tester(PCGAttributeTestsCommonHelper::StringTester());
 	}
 
 	SECTION("Vector")
 	{
-		Tester(PCGAttributeGenericTests::VectorTester_Basic());
+		Tester(PCGAttributeTestsCommonHelper::VectorTester_Basic());
 	}
 
 	SECTION("FPCGPoint (Struct)")
 	{
-		Tester(PCGAttributeGenericTests::FPCGPointTester());
+		Tester(PCGAttributeTestsCommonHelper::FPCGPointTester());
 	}
 }
 
@@ -1260,7 +917,7 @@ TEST_CASE_METHOD(PCGTests::FPCGBaseTestWithBasicContext, "PCG::AttributeGeneric:
  */
 TEST_CASE_METHOD(PCGTests::FPCGBaseTestWithBasicContext, "PCG::AttributeGeneric::GetValuesContinuous", "[PCG][AttributeGeneric]")
 {
-	auto Tester = [this]<typename T, typename ReadType>(PCGAttributeGenericTests::TypedAttributeTester<T, ReadType> AttributeTester)
+	auto Tester = [this]<typename T, typename ReadType>(PCGAttributeTestsCommonHelper::TypedAttributeTester<T, ReadType> AttributeTester)
 	{
 		// Preparing the data
 		constexpr int32 NumPoints = 100;
@@ -1277,7 +934,7 @@ TEST_CASE_METHOD(PCGTests::FPCGBaseTestWithBasicContext, "PCG::AttributeGeneric:
 
 		FPCGMetadataAttributeGeneric* GenericAttribute = nullptr;
 
-		GenericAttribute = Domain->CreateGenericAttribute<T>(PCGAttributeGenericTests::AttributeName, AttributeTester.DefaultValue);
+		GenericAttribute = Domain->CreateGenericAttribute<T>(PCGAttributeTestsCommonHelper::AttributeName, AttributeTester.DefaultValue);
 		REQUIRE_NOT_EQUAL(GenericAttribute, nullptr);
 
 		TStaticArray<T, NumPoints> Values{};
@@ -1309,7 +966,7 @@ TEST_CASE_METHOD(PCGTests::FPCGBaseTestWithBasicContext, "PCG::AttributeGeneric:
 
 			TPCGValueRange<int64> ChildPointArrayDataMetadataEntryRange = ChildPointArrayData->GetMetadataEntryValueRange(/*bAllocate=*/false);
 
-			FPCGMetadataAttributeGeneric* ChildGenericAttribute = ChildDomain->GetMutableGenericAttribute(PCGAttributeGenericTests::AttributeName);
+			FPCGMetadataAttributeGeneric* ChildGenericAttribute = ChildDomain->GetMutableGenericAttribute(PCGAttributeTestsCommonHelper::AttributeName);
 			REQUIRE_NOT_EQUAL(ChildGenericAttribute, nullptr);
 
 			// Add more values
@@ -1338,51 +995,51 @@ TEST_CASE_METHOD(PCGTests::FPCGBaseTestWithBasicContext, "PCG::AttributeGeneric:
 
 	SECTION("Float Attribute")
 	{
-		Tester(PCGAttributeGenericTests::FloatTester());
+		Tester(PCGAttributeTestsCommonHelper::FloatTester());
 	}
 
 	SECTION("Double Attribute")
 	{
-		Tester(PCGAttributeGenericTests::DoubleTester());
+		Tester(PCGAttributeTestsCommonHelper::DoubleTester());
 	}
 
 	SECTION("As Basic Type")
 	{
-		Tester(PCGAttributeGenericTests::VectorTester_Basic());
+		Tester(PCGAttributeTestsCommonHelper::VectorTester_Basic());
 	}
 
 	SECTION("String Attribute")
 	{
-		Tester(PCGAttributeGenericTests::StringTester());
+		Tester(PCGAttributeTestsCommonHelper::StringTester());
 	}
 
 	SECTION("FPCGPoint Attribute")
 	{
-		Tester(PCGAttributeGenericTests::FPCGPointTester());
+		Tester(PCGAttributeTestsCommonHelper::FPCGPointTester());
 	}
 
 	SECTION("Array of Double Attribute")
 	{
-		Tester(PCGAttributeGenericTests::ArrayDoubleTester());
+		Tester(PCGAttributeTestsCommonHelper::ArrayDoubleTester());
 	}
 
 	SECTION("Array of FPCGPoint Attribute")
 	{
-		Tester(PCGAttributeGenericTests::ArrayFPCGPointTester());
+		Tester(PCGAttributeTestsCommonHelper::ArrayFPCGPointTester());
 	}
 
 	SECTION("Set of Double Attribute")
 	{
-		Tester(PCGAttributeGenericTests::SetDoubleTester());
+		Tester(PCGAttributeTestsCommonHelper::SetDoubleTester());
 	}
 
 	SECTION("Map of String -> Double Attribute")
 	{
-		Tester(PCGAttributeGenericTests::MapStringDoubleTester());
+		Tester(PCGAttributeTestsCommonHelper::MapStringDoubleTester());
 	}
 
 	SECTION("Enum attribute")
 	{
-		Tester(PCGAttributeGenericTests::EnumTester());
+		Tester(PCGAttributeTestsCommonHelper::EnumTester());
 	}
 }
\ No newline at end of file
diff --git a/Engine/Plugins/PCG/Tests/Private/Metadata/PCGMetadataAttributeTestsCommonHelper.h b/Engine/Plugins/PCG/Tests/Private/Metadata/PCGMetadataAttributeTestsCommonHelper.h
new file mode 100644
--- /dev/null
+++ b/Engine/Plugins/PCG/Tests/Private/Metadata/PCGMetadataAttributeTestsCommonHelper.h
@@ -0,0 +1,351 @@
+﻿// Copyright Epic Games, Inc. All Rights Reserved.
+#pragma once
+
+#include "PCGTestsCommon.h"
+
+#include "PCGCommon.h"
+#include "Metadata/PCGMetadataContainerTypes.h"
+
+namespace PCGAttributeTestsCommonHelper
+{
+	/**
+	 * Base class for a tester that encapsulate everything that is type specific.
+	 */
+	template <typename T, typename InReadType = T>
+	struct TypedAttributeTester
+	{
+		using Type = T;
+		using ReadType = InReadType;
+
+		TypedAttributeTester(TFunction<void(FRandomStream&, T&)> InGenerateRandom, 
+			TFunction<bool(const T&, const ReadType&)> InVerify, 
+			T InDefaultValue,
+			FPCGMetadataAttributeDesc InExpectedDesc)
+			: GenerateRandom(MoveTemp(InGenerateRandom))
+			, Verify(MoveTemp(InVerify))
+			, DefaultValue(InDefaultValue)
+			, ExpectedDesc(InExpectedDesc)
+		{}
+
+		TFunction<void(FRandomStream&, T&)> GenerateRandom;
+		TFunction<bool(const T&, const ReadType&)> Verify;
+		T DefaultValue;
+		FPCGMetadataAttributeDesc ExpectedDesc;
+		TOptional<FPCGMetadataAttributeDesc> OverrideDesc;
+	};
+	
+	static const FName AttributeName = "Attr";
+	static const FName AttributeName2 = "Attr2";
+	static const FName AttributeName3 = "Attr3";
+	
+	/** Getter to create a tester for different types. */
+	inline TypedAttributeTester<float> FloatTester()
+	{
+		auto GenerateRandom = [](FRandomStream& RandomStream, float& OutValue) { OutValue = RandomStream.FRand(); };
+		
+		auto Verify = [](const float& LHS, const float& RHS) { return LHS == RHS; };
+		
+		float DefaultValue = 5.0f;
+		
+		FPCGMetadataAttributeDesc ExpectedDesc
+		{
+			.Name = AttributeName,
+			.ValueType = EPCGMetadataTypes::Float
+		};
+		
+		return TypedAttributeTester<float>(GenerateRandom, Verify, DefaultValue, MoveTemp(ExpectedDesc));
+	}
+	
+	inline TypedAttributeTester<double> DoubleTester()
+	{
+		auto GenerateRandom = [](FRandomStream& RandomStream, double& OutValue) { OutValue = RandomStream.FRand(); };
+		
+		auto Verify = [](const double& LHS, const double& RHS) { return LHS == RHS; };
+		
+		double DefaultValue = 5.0;
+		
+		FPCGMetadataAttributeDesc ExpectedDesc
+		{
+			.Name = AttributeName,
+			.ValueType = EPCGMetadataTypes::Double
+		};
+		
+		return TypedAttributeTester<double>(GenerateRandom, Verify, DefaultValue, MoveTemp(ExpectedDesc));
+	}
+	
+	inline TypedAttributeTester<FVector> VectorTester_Basic()
+	{
+		auto GenerateRandom = [](FRandomStream& RandomStream, FVector& OutValue) { OutValue = RandomStream.VRand(); };
+		
+		auto Verify = [](const FVector& LHS, const FVector& RHS) { return LHS == RHS; };
+		
+		FVector DefaultValue = FVector(1.0, 2.0, 3.0);
+		
+		FPCGMetadataAttributeDesc ExpectedDesc
+		{
+			.Name = AttributeName,
+			.ValueType = EPCGMetadataTypes::Vector
+		};
+			
+		return TypedAttributeTester<FVector>{GenerateRandom, Verify, DefaultValue, MoveTemp(ExpectedDesc)};
+	}
+		
+	inline TypedAttributeTester<FVector> VectorTester_Struct()
+	{
+		TypedAttributeTester<FVector> AttributeTester = VectorTester_Basic();
+		FPCGMetadataAttributeDesc StructDesc
+		{
+			.Name = AttributeName,
+			.ValueType = EPCGMetadataTypes::Struct,
+			.ValueTypeObject = TBaseStructure<FVector>::Get()
+		};
+		
+		AttributeTester.OverrideDesc = StructDesc;
+	
+		return AttributeTester;
+	}
+	
+	inline TypedAttributeTester<FString, const FString*> StringTester()
+	{
+		auto GenerateRandom = [](FRandomStream& RandomStream, FString& OutValue)
+			{
+				static int i = 0; 
+				OutValue = FString::Printf(TEXT("This is a string with some characters and a value: %d"), i++ % 25);
+			};
+		
+		auto Verify = [](const FString& LHS, FString const* const& RHS) { return LHS == *RHS; };
+		
+		FString DefaultValue = TEXT("Hi");
+		
+		FPCGMetadataAttributeDesc ExpectedDesc
+		{
+			.Name = AttributeName,
+			.ValueType = EPCGMetadataTypes::String
+		};
+		
+		return TypedAttributeTester<FString, const FString*>(GenerateRandom, Verify, MoveTemp(DefaultValue), MoveTemp(ExpectedDesc));
+	}
+	
+	inline TypedAttributeTester<FPCGPoint, const FPCGPoint*> FPCGPointTester()
+	{		
+		auto GenerateRandom = [](FRandomStream& RandomStream, FPCGPoint& OutValue)
+			{
+				OutValue.Density = RandomStream.FRand();
+				OutValue.Transform.SetLocation(RandomStream.VRand());
+			};
+		
+		auto Verify = [](const FPCGPoint& LHS, FPCGPoint const* const& RHS)
+			{
+				return LHS.Density == RHS->Density && LHS.Transform.GetLocation() == RHS->Transform.GetLocation();
+			};
+		
+		FPCGPoint DefaultValue{FTransform{5.0 * FVector::OneVector}, 0.4f, 42};
+		
+		FPCGMetadataAttributeDesc ExpectedDesc
+		{
+			.Name = AttributeName,
+			.ValueType = EPCGMetadataTypes::Struct,
+			.ValueTypeObject = FPCGPoint::StaticStruct()
+		};
+		
+		return TypedAttributeTester<FPCGPoint, const FPCGPoint*>(GenerateRandom, Verify, MoveTemp(DefaultValue), MoveTemp(ExpectedDesc));
+	}
+	
+	inline TypedAttributeTester<TArray<double>, TConstArrayView<double>> ArrayDoubleTester()
+	{		
+		auto GenerateRandom = [](FRandomStream& RandomStream, TArray<double>& OutValue)
+		{
+			const int32 RandomCount = RandomStream.RandRange(1, 4);
+			for (int j = 0; j < RandomCount; ++j)
+			{
+				OutValue.Add(RandomStream.FRand());
+			}
+		};
+		
+		auto Verify = [](const TArray<double>& LHS, const TConstArrayView<double>& RHS)
+			{
+				if (LHS.Num() != RHS.Num())
+				{
+					return false;
+				}
+				
+				for (int i = 0; i < LHS.Num(); ++i)
+				{
+					if (LHS[i] != RHS[i])
+					{
+						return false;
+					}
+				}
+				
+				return true;
+			};
+		
+		TArray<double> DefaultValue = {1, 2, 3, 4, 5};
+		
+		FPCGMetadataAttributeDesc ExpectedDesc
+		{
+			.Name = AttributeName,
+			.ValueType = EPCGMetadataTypes::Double,
+			.ContainerTypes = {EPCGMetadataAttributeContainerTypes::Array}
+		};
+		
+		return TypedAttributeTester<TArray<double>, TConstArrayView<double>>(GenerateRandom, Verify, MoveTemp(DefaultValue), MoveTemp(ExpectedDesc));
+	}
+	
+	inline TypedAttributeTester<TArray<FPCGPoint>, TConstArrayView<FPCGPoint>> ArrayFPCGPointTester()
+	{	
+		auto GenerateRandom = [](FRandomStream& RandomStream, TArray<FPCGPoint>& OutValue)
+			{
+				const int32 RandomCount = RandomStream.RandRange(1, 4);
+				for (int j = 0; j < RandomCount; ++j)
+				{
+					FPCGPoint& Value = OutValue.Emplace_GetRef();
+					Value.Density = RandomStream.FRand();
+					Value.Transform.SetLocation(RandomStream.VRand());
+				}
+			};
+		
+		auto Verify = [](const TArray<FPCGPoint>& LHS, const TConstArrayView<FPCGPoint>& RHS)
+			{
+				if (LHS.Num() != RHS.Num())
+				{
+					return false;
+				}
+				
+				for (int i = 0; i < LHS.Num(); ++i)
+				{
+					if (LHS[i].Density != RHS[i].Density || LHS[i].Transform.GetLocation() != RHS[i].Transform.GetLocation())
+					{
+						return false;
+					}
+				}
+				
+				return true;
+			};
+		
+		TArray<FPCGPoint> DefaultValue = {
+			{FTransform{5.0 * FVector::OneVector}, 0.4f, 42},
+			{FTransform{6.0 * FVector::OneVector}, 0.5f, 43},
+			{FTransform{7.0 * FVector::OneVector}, 0.6f, 44},
+		};
+		
+		FPCGMetadataAttributeDesc ExpectedDesc
+		{
+			.Name = AttributeName,
+			.ValueType = EPCGMetadataTypes::Struct,
+			.ContainerTypes = {EPCGMetadataAttributeContainerTypes::Array},
+			.ValueTypeObject = FPCGPoint::StaticStruct()
+		};
+		
+		return TypedAttributeTester<TArray<FPCGPoint>, TConstArrayView<FPCGPoint>>(GenerateRandom, Verify, MoveTemp(DefaultValue), MoveTemp(ExpectedDesc));
+	}
+	
+	inline TypedAttributeTester<TSet<double>, PCG::TScriptSetWrapper<double>> SetDoubleTester()
+	{		
+		auto GenerateRandom = [](FRandomStream& RandomStream, TSet<double>& OutValue)
+			{
+				const int32 RandomCount = RandomStream.RandRange(1, 4);
+				for (int j = 0; j < RandomCount; ++j)
+				{
+					OutValue.Add(RandomStream.FRand());
+				}
+			};
+		
+		auto Verify = [](const TSet<double>& LHS, const PCG::TScriptSetWrapper<double>& RHS)
+			{
+				if (LHS.Num() != RHS.Num())
+				{
+					return false;
+				}
+				
+				for (double Value : LHS)
+				{
+					if (!RHS.Contains(Value))
+					{
+						return false;
+					}
+				}
+				
+				return true;
+			};
+		
+		TSet<double> DefaultValue = {1, 2, 3, 4, 5};
+		
+		FPCGMetadataAttributeDesc ExpectedDesc
+		{
+			.Name = AttributeName,
+			.ValueType = EPCGMetadataTypes::Double,
+			.ContainerTypes = {EPCGMetadataAttributeContainerTypes::Set}
+		};
+		
+		return TypedAttributeTester<TSet<double>, PCG::TScriptSetWrapper<double>>(GenerateRandom, Verify, MoveTemp(DefaultValue), MoveTemp(ExpectedDesc));
+	}
+	
+	inline TypedAttributeTester<TMap<FString, double>, PCG::TScriptMapWrapper<FString, double>> MapStringDoubleTester()
+	{
+		auto GenerateRandom = [](FRandomStream& RandomStream, TMap<FString, double>& OutValue)
+			{
+				const int32 RandomCount = RandomStream.RandRange(1, 4);
+				for (int j = 0; j < RandomCount; ++j)
+				{
+					OutValue.Emplace(LexToString(j), RandomStream.FRand());
+				}
+			};
+		
+		auto Verify = [](const TMap<FString, double>& LHS, const PCG::TScriptMapWrapper<FString, double>& RHS)
+			{
+				if (LHS.Num() != RHS.Num())
+				{
+					return false;
+				}
+				
+				for (const auto& [Key, Value] : LHS)
+				{
+					const double* OtherValue = RHS.Find(Key);
+					if (!OtherValue || Value != *OtherValue)
+					{
+						return false;
+					}
+				}
+				
+				return true;
+			};
+		
+		TMap<FString, double> DefaultValue =
+		{
+			{TEXT("Hi"), 2.0},
+			{TEXT("Ho"), 3.0}
+		};
+		
+		FPCGMetadataAttributeDesc ExpectedDesc
+		{
+			.Name = AttributeName,
+			.ValueType = EPCGMetadataTypes::Double,
+			.ContainerTypes = {EPCGMetadataAttributeContainerTypes::Map},
+			.KeyType = EPCGMetadataTypes::String
+		};
+		
+		return TypedAttributeTester<TMap<FString, double>, PCG::TScriptMapWrapper<FString, double>>(GenerateRandom, Verify, MoveTemp(DefaultValue), MoveTemp(ExpectedDesc));
+	}
+	
+	inline TypedAttributeTester<EPCGMetadataTypes> EnumTester()
+	{
+		auto GenerateRandom = [](FRandomStream& RandomStream, EPCGMetadataTypes& OutValue)
+			{
+				OutValue = static_cast<EPCGMetadataTypes>(RandomStream.RandRange(0, static_cast<int32>(EPCGMetadataTypes::Count)));
+			};
+		
+		auto Verify = [](const EPCGMetadataTypes& LHS, const EPCGMetadataTypes& RHS) { return LHS == RHS; };
+		
+		EPCGMetadataTypes DefaultValue = EPCGMetadataTypes::Unknown;
+		
+		FPCGMetadataAttributeDesc ExpectedDesc
+		{
+			.Name = AttributeName,
+			.ValueType = EPCGMetadataTypes::Enum,
+			.ValueTypeObject = StaticEnum<EPCGMetadataTypes>()
+		};
+		
+		return TypedAttributeTester<EPCGMetadataTypes>(GenerateRandom, Verify, DefaultValue, MoveTemp(ExpectedDesc));
+	}
+}
diff --git a/Engine/Plugins/PCG/Tests/Private/PCGTestsCommon.h b/Engine/Plugins/PCG/Tests/Private/PCGTestsCommon.h
--- a/Engine/Plugins/PCG/Tests/Private/PCGTestsCommon.h
+++ b/Engine/Plugins/PCG/Tests/Private/PCGTestsCommon.h
@@ -1,5 +1,7 @@
 ﻿// Copyright Epic Games, Inc. All Rights Reserved.
 
+#pragma once
+
 #include "PCGContext.h"
 
 #include "CoreMinimal.h"
