[PCG] Add support to any type of data in generic attribute accessors

#rnx
#Synced-CL 52060281

diff --git a/Engine/Plugins/PCG/Source/PCG/Private/Metadata/Accessors/IPCGAttributeAccessor.cpp b/Engine/Plugins/PCG/Source/PCG/Private/Metadata/Accessors/IPCGAttributeAccessor.cpp
--- a/Engine/Plugins/PCG/Source/PCG/Private/Metadata/Accessors/IPCGAttributeAccessor.cpp
+++ b/Engine/Plugins/PCG/Source/PCG/Private/Metadata/Accessors/IPCGAttributeAccessor.cpp
@@ -25,7 +25,7 @@ bool IPCGAttributeAccessor::GetRange_Internal(PCG::Private::FOutValues OutValues
 	{
 	case EPCGAttributeAccessorOperation::SameType:
 	{
-		return GetRangeVirtual(MoveTemp(OutValues), Count, Index, Keys, Flags);
+		return GetRangeVirtual(MoveTemp(OutValues), Count, Index, Keys);
 	}
 	case EPCGAttributeAccessorOperation::Construct: // Fall-through
 	case EPCGAttributeAccessorOperation::Broadcast:
@@ -47,7 +47,7 @@ bool IPCGAttributeAccessor::GetRange_Internal(PCG::Private::FOutValues OutValues
 			(*ConstructFunc)(Buffer.GetData(), Count);
 		}
 
-		const bool bSuccess = GetRangeVirtual(FOutValues{TInPlaceType<FOutValuesByValue>{}, Buffer.GetData()}, Count, Index, Keys, Flags);
+		const bool bSuccess = GetRangeVirtual(FOutValues{TInPlaceType<FOutValuesByValue>{}, Buffer.GetData()}, Count, Index, Keys);
 			
 		if (bSuccess)
 		{
diff --git a/Engine/Plugins/PCG/Source/PCG/Private/Metadata/Accessors/PCGAttributeAccessorHelpers.cpp b/Engine/Plugins/PCG/Source/PCG/Private/Metadata/Accessors/PCGAttributeAccessorHelpers.cpp
--- a/Engine/Plugins/PCG/Source/PCG/Private/Metadata/Accessors/PCGAttributeAccessorHelpers.cpp
+++ b/Engine/Plugins/PCG/Source/PCG/Private/Metadata/Accessors/PCGAttributeAccessorHelpers.cpp
@@ -376,7 +376,15 @@ TUniquePtr<const IPCGAttributeAccessor> PCGAttributeAccessorHelpers::CreateConst
 			}
 		};
 		
-		return PCGMetadataAttribute::CallbackWithRightType(InAttribute->GetTypeId(), MakeConstAccessor);
+		TUniquePtr<const IPCGAttributeAccessor> OutAccessor = PCGMetadataAttribute::CallbackWithRightType(InAttribute->GetTypeId(), MakeConstAccessor);
+		if (!OutAccessor && InAttribute->IsGeneric())
+		{
+			// @todo_pcg: To be updated when we won't have typed accessors anymore.
+			// This creates a generic accessor that knows its type with the desc in InAttribute.
+			OutAccessor = MakeUnique<const FPCGAttributeGenericAccessor<uint8>>(static_cast<const FPCGMetadataAttributeGeneric*>(InAttribute), InMetadata);
+		}
+		
+		return OutAccessor;
 	}
 }
 
@@ -459,8 +467,16 @@ TUniquePtr<IPCGAttributeAccessor> PCGAttributeAccessorHelpers::CreateAccessor(FP
 				return MakeUnique<FPCGAttributeAccessor<T>>(static_cast<FPCGMetadataAttribute<T>*>(InAttribute), InMetadata);
 			}
 		};
-		
-		return PCGMetadataAttribute::CallbackWithRightType(InAttribute->GetTypeId(), MakeAccessor);
+
+		TUniquePtr<IPCGAttributeAccessor> OutAccessor = PCGMetadataAttribute::CallbackWithRightType(InAttribute->GetTypeId(), MakeAccessor);
+		if (!OutAccessor && InAttribute->IsGeneric())
+		{
+			// @todo_pcg: To be updated when we won't have typed accessors anymore.
+			// This creates a generic accessor that knows its type with the desc in InAttribute.
+			OutAccessor = MakeUnique<FPCGAttributeGenericAccessor<uint8>>(static_cast<FPCGMetadataAttributeGeneric*>(InAttribute), InMetadata);
+		}
+
+		return OutAccessor;
 	}
 }
 
diff --git a/Engine/Plugins/PCG/Source/PCG/Public/Metadata/Accessors/IPCGAttributeAccessor.h b/Engine/Plugins/PCG/Source/PCG/Public/Metadata/Accessors/IPCGAttributeAccessor.h
--- a/Engine/Plugins/PCG/Source/PCG/Public/Metadata/Accessors/IPCGAttributeAccessor.h
+++ b/Engine/Plugins/PCG/Source/PCG/Public/Metadata/Accessors/IPCGAttributeAccessor.h
@@ -494,7 +494,7 @@ protected:
 	}
 	
 	PCG_API bool GetRange_Internal(PCG::Private::FOutValues OutValues, int16 ExpectedType, int32 Count, int32 Index, const IPCGAttributeAccessorKeys& Keys, EPCGAttributeAccessorFlags Flags, PCG::Private::EPCGAttributeAccessorOperation OpType) const;
-	virtual bool GetRangeVirtual(PCG::Private::FOutValues OutValues, int32 Count, int32 Index, const IPCGAttributeAccessorKeys& Keys, EPCGAttributeAccessorFlags Flags) const = 0;
+	virtual bool GetRangeVirtual(PCG::Private::FOutValues OutValues, int32 Count, int32 Index, const IPCGAttributeAccessorKeys& Keys) const = 0;
 	
 	PCG_API bool SetRange_Internal(PCG::Private::FInValues InValues, int16 ExpectedType, int32 Count, int32 Index, IPCGAttributeAccessorKeys& Keys, EPCGAttributeAccessorFlags Flags, PCG::Private::EPCGAttributeAccessorOperation OpType);
 	virtual bool SetRangeVirtual(PCG::Private::FInValues InValues, int32 Count, int32 Index, IPCGAttributeAccessorKeys& Keys, EPCGAttributeAccessorFlags Flags) = 0;
diff --git a/Engine/Plugins/PCG/Source/PCG/Public/Metadata/Accessors/IPCGAttributeAccessorTpl.h b/Engine/Plugins/PCG/Source/PCG/Public/Metadata/Accessors/IPCGAttributeAccessorTpl.h
--- a/Engine/Plugins/PCG/Source/PCG/Public/Metadata/Accessors/IPCGAttributeAccessorTpl.h
+++ b/Engine/Plugins/PCG/Source/PCG/Public/Metadata/Accessors/IPCGAttributeAccessorTpl.h
@@ -24,7 +24,7 @@ protected:
 	{}
 	
 	// Default implementation
-	virtual bool GetRangeVirtual(PCG::Private::FOutValues OutValues, int32 Count, int32 Index, const IPCGAttributeAccessorKeys& Keys, EPCGAttributeAccessorFlags Flags) const override
+	virtual bool GetRangeVirtual(PCG::Private::FOutValues OutValues, int32 Count, int32 Index, const IPCGAttributeAccessorKeys& Keys) const override
 	{
 		using T = typename Derived::Type;
 		const Derived* This = static_cast<const Derived*>(this);
diff --git a/Engine/Plugins/PCG/Source/PCG/Public/Metadata/Accessors/PCGAttributeAccessor.h b/Engine/Plugins/PCG/Source/PCG/Public/Metadata/Accessors/PCGAttributeAccessor.h
--- a/Engine/Plugins/PCG/Source/PCG/Public/Metadata/Accessors/PCGAttributeAccessor.h
+++ b/Engine/Plugins/PCG/Source/PCG/Public/Metadata/Accessors/PCGAttributeAccessor.h
@@ -256,11 +256,92 @@ public:
 	FPCGAttributeGenericAccessor(FPCGMetadataAttributeGeneric* InAttribute, FPCGMetadataDomain* InMetadataDomain, bool bForceReadOnly = false)
 		: Super(InAttribute, InMetadataDomain, bForceReadOnly)
 	{
+		this->UnderlyingDesc = InAttribute->GetAttributeDesc();
+		this->UnderlyingType = static_cast<int16>(this->UnderlyingDesc.ValueType);
 	}
 
 	FPCGAttributeGenericAccessor(const FPCGMetadataAttributeGeneric* InAttribute, const FPCGMetadataDomain* InMetadataDomain, bool bForceReadOnly = false)
 		: Super(InAttribute, InMetadataDomain, bForceReadOnly)
 	{
+		this->UnderlyingDesc = InAttribute->GetAttributeDesc();
+		this->UnderlyingType = static_cast<int16>(this->UnderlyingDesc.ValueType);
+	}
+	
+	virtual bool GetRangeVirtual(PCG::Private::FOutValues OutValues, int32 Count, int32 Index, const IPCGAttributeAccessorKeys& Keys) const override
+	{
+		const FPCGMetadataAttributeGeneric* ThisAttribute = static_cast<const FPCGMetadataAttributeGeneric*>(this->Attribute);
+		
+		TArray<const PCGMetadataEntryKey*, TInlineAllocator<256>> EntryKeyPtrs;
+		EntryKeyPtrs.SetNumUninitialized(Count);
+
+		TArrayView<const PCGMetadataEntryKey*> EntryKeysView(EntryKeyPtrs);
+		if (!Keys.GetKeys<PCGMetadataEntryKey>(Index, EntryKeysView))
+		{
+			return false;
+		}
+
+		TArray<PCGMetadataEntryKey, TInlineAllocator<256>> EntryKeys;
+		EntryKeys.Reserve(EntryKeyPtrs.Num());
+		Algo::Transform(EntryKeyPtrs, EntryKeys, [](const PCGMetadataEntryKey* KeyPtr) { return *KeyPtr; });
+		
+		TArray<PCGMetadataValueKey> ValueKeys;
+		ThisAttribute->GetValueKeys(PCGValueRangeHelpers::MakeConstValueRange<PCGMetadataEntryKey>(EntryKeys), ValueKeys);
+		ThisAttribute->GetValues_Internal(ValueKeys, MoveTemp(OutValues));
+
+		return true;
+	}
+	
+	virtual bool SetRangeVirtual(PCG::Private::FInValues InValues, int32 Count, int32 Index, IPCGAttributeAccessorKeys& Keys, EPCGAttributeAccessorFlags Flags) override
+	{
+		FPCGMetadataAttributeGeneric* ThisAttribute = static_cast<FPCGMetadataAttributeGeneric*>(this->Attribute);
+		
+		if (EnumHasAnyFlags(Flags, EPCGAttributeAccessorFlags::AllowSetDefaultValue) && Keys.GetNum() == 1)
+		{
+			PCGMetadataEntryKey* EntryKey = nullptr;
+			Keys.GetKey(EntryKey);
+
+			if (EntryKey && *EntryKey == PCGInvalidEntryKey)
+			{
+				check(Count > 0);
+				ThisAttribute->AddValues_Internal(InValues, 1, /*bExistingIndex=*/PCG::Private::ExistingIndexForDefaultValue);
+				return true;
+			}
+		}
+
+		TArray<PCGMetadataEntryKey*, TInlineAllocator<512>> EntryKeys;
+		EntryKeys.SetNumUninitialized(Count);
+		TArrayView<PCGMetadataEntryKey*> EntryKeysView(EntryKeys);
+
+		// Not const because it can be modified by Prepare
+		int32 Start = Index;
+
+		if (!this->bWasPrepared || ThisAttribute->DoesCompressData())
+		{
+			if (!Super::Prepare(Keys, Count, /*bCanReuseEntryKeys=*/EnumHasAnyFlags(Flags, EPCGAttributeAccessorFlags::AllowReuseMetadataEntryKey), Index, /*bPreallocateValues=*/true, &EntryKeysView, &Start))
+			{
+				return false;
+			}
+
+			TArray<PCGMetadataValueKey> ValueKeys;
+			
+			if (ThisAttribute->DoesCompressData() && ensure(InValues.IsType<PCG::Private::FInValuesByValue>()))
+			{
+				ValueKeys = ThisAttribute->AddCompressedValues(InValues.Get<PCG::Private::FInValuesByValue>(), Count);
+			}
+			else
+			{
+				ValueKeys = ThisAttribute->AddValues_Internal(InValues, Count, /*bExistingIndex=*/Start);
+			}
+			
+			ThisAttribute->SetValuesFromValueKeys(EntryKeys, ValueKeys);
+		}
+		else
+		{
+			check(this->StartIndex != INDEX_NONE);
+			ThisAttribute->AddValues_Internal(InValues, Count, /*bExistingIndex=*/this->StartIndex + Start, /*bLockless=*/true);
+		}
+
+		return true;
 	}
 	
 	void GetValuesFromItemKeys(TArrayView<PCGMetadataEntryKey> EntryKeys, TArrayView<T> OutValues) const
diff --git a/Engine/Plugins/PCG/Tests/Private/Metadata/Accessors/PCGMetadataAttributeGenericAccessorTests.cpp b/Engine/Plugins/PCG/Tests/Private/Metadata/Accessors/PCGMetadataAttributeGenericAccessorTests.cpp
--- a/Engine/Plugins/PCG/Tests/Private/Metadata/Accessors/PCGMetadataAttributeGenericAccessorTests.cpp
+++ b/Engine/Plugins/PCG/Tests/Private/Metadata/Accessors/PCGMetadataAttributeGenericAccessorTests.cpp
@@ -1,18 +1,175 @@
-﻿// Copyright Epic Games, Inc. All Rights Reserved.
+// Copyright Epic Games, Inc. All Rights Reserved.
+
+/**
+ * Tests for FPCGAttributeGenericAccessor -- the type-erased accessor for generic attributes.
+ * These accessors bypass the CRTP typed path and instead use the virtual FInValues/FOutValues
+ * variant dispatch to read/write data of any type supported by the attribute system.
+ *
+ * Test cases:
+ *   Basic          - Round-trip write/read for all supported types (scalars, containers, structs, enums).
+ *   NotExactType   - Type broadcast and construction: writing int32 into a double attribute,
+ *                    then reading back as int32 and FString.
+ *   CompressedData - Value deduplication for compressed attributes (e.g. FString), verifying that
+ *                    repeated values share storage and new values are correctly appended.
+ */
 
 #include "CoreMinimal.h"
 #include "TestHarness.h"
 
+#include "PCGTestsCommon.h"
+#include "Metadata/PCGMetadataAttributeTestsCommonHelper.h"
+
 #include <catch2/generators/catch_generators.hpp>
 
-#include "PCGContext.h"
 #include "Data/PCGPointArrayData.h"
-#include "Tests/PCGTestsCommon.h"
-#include "PCGTestsCommon.h"
 #include "Metadata/Accessors/PCGAttributeAccessorHelpers.h"
 
 
-TEST_CASE_METHOD(PCGTests::FPCGBaseTestWithBasicContext, "PCG::Metadata::Accessors::GenericAttribute", "[PCG][Accessors][GenericAttribute]")
+/**
+ * Validates that FPCGAttributeGenericAccessor can write and read back values for every supported type.
+ * Each type is tested via a TypedAttributeTester that generates random values, writes them through
+ * SetRange, and reads them back through GetRange, comparing the results. The write is split into two
+ * halves to exercise offset-based range access. Types covered: Float, Double, Vector (basic and struct
+ * representations), String, FPCGPoint (struct), Array<double>, Array<FPCGPoint>, Set<double>,
+ * Map<FString, double>, and Enum.
+ */
+TEST_CASE_METHOD(PCGTests::FPCGBaseTestWithBasicContext, "PCG::Metadata::Accessors::GenericAttribute::Basic", "[PCG][AttributeGeneric]")
+{
+	constexpr int32 Seed = 42;
+	
+	auto Tester = [this]<typename T, typename ReadType>(PCGAttributeTestsCommonHelper::TypedAttributeTester<T, ReadType> AttributeTester)
+	{
+		// Preparing the data
+		constexpr int32 NumPoints = 100;
+		constexpr int32 HalfNumPoints = NumPoints / 2;
+		UPCGPointArrayData* PointArrayData = CreateData<UPCGPointArrayData>();
+
+		PointArrayData->SetNumPoints(NumPoints);
+
+		PointArrayData->AllocateProperties(EPCGPointNativeProperties::MetadataEntry);
+		TPCGValueRange<int64> PointArrayDataMetadataEntryRange = PointArrayData->GetMetadataEntryValueRange(/*bAllocate=*/false);
+		
+		FPCGMetadataDomain* Domain = PointArrayData->MutableMetadata()->GetDefaultMetadataDomain();
+		
+		FRandomStream RandomStream(Seed);
+		
+		FPCGMetadataAttributeGeneric* GenericAttribute = nullptr;
+	
+		if (AttributeTester.OverrideDesc.IsSet())
+		{
+			GenericAttribute = Domain->CreateGenericAttribute(*AttributeTester.OverrideDesc);
+			REQUIRE_NOT_EQUAL(GenericAttribute, nullptr);
+			GenericAttribute->SetDefaultValue<T>(AttributeTester.DefaultValue);
+		}
+		else
+		{
+			GenericAttribute = Domain->CreateGenericAttribute<T>(PCGAttributeTestsCommonHelper::AttributeName, AttributeTester.DefaultValue);
+			REQUIRE_NOT_EQUAL(GenericAttribute, nullptr);
+		}
+		
+		SECTION("Expected attribute descriptor")
+		{
+			REQUIRE_EQUAL(GenericAttribute->GetAttributeDesc(), AttributeTester.ExpectedDesc);
+		}
+		
+		TUniquePtr<IPCGAttributeAccessor> AttributeAccessor = PCGAttributeAccessorHelpers::CreateAccessor(PointArrayData, FPCGAttributePropertySelector::CreateAttributeSelector(PCGAttributeTestsCommonHelper::AttributeName));
+		TUniquePtr<IPCGAttributeAccessorKeys> AttributeKeys = PCGAttributeAccessorHelpers::CreateKeys(PointArrayData, FPCGAttributePropertySelector::CreateAttributeSelector(PCGAttributeTestsCommonHelper::AttributeName));
+
+		REQUIRE(AttributeAccessor.IsValid());
+		REQUIRE(AttributeKeys.IsValid());
+		
+		TStaticArray<T, NumPoints> Values{};
+		TStaticArray<ReadType, NumPoints> ReadValues;
+		
+		for (int32 i = 0; i < NumPoints; ++i)
+		{
+			AttributeTester.GenerateRandom(RandomStream, Values[i]);
+		}
+
+		// Do it twice, on half the points
+		REQUIRE(AttributeAccessor->SetRange<T>(MakeArrayView(Values.GetData(), HalfNumPoints), 0, *AttributeKeys));
+		REQUIRE_EQUAL(GenericAttribute->GetNumberOfEntries(), HalfNumPoints);
+
+		REQUIRE(AttributeAccessor->SetRange<T>(MakeArrayView(Values.GetData() + HalfNumPoints, HalfNumPoints), HalfNumPoints, *AttributeKeys));
+		REQUIRE_EQUAL(GenericAttribute->GetNumberOfEntries(), NumPoints);
+		
+		SECTION("Read back and verify values")
+		{
+			// Do it twice, on half the points
+			REQUIRE(AttributeAccessor->GetRange<ReadType>(MakeArrayView(ReadValues.GetData(), HalfNumPoints), 0, *AttributeKeys));
+			REQUIRE(AttributeAccessor->GetRange<ReadType>(MakeArrayView(ReadValues.GetData() + HalfNumPoints, HalfNumPoints), HalfNumPoints, *AttributeKeys));
+
+			CHECK_THAT(Values, Catch::Matchers::RangeEquals(ReadValues, [&AttributeTester](const auto& a, const auto& b){ return AttributeTester.Verify(a, b); } ));
+		}
+	};
+	
+	SECTION("Float Attribute")
+	{
+		Tester(PCGAttributeTestsCommonHelper::FloatTester());
+	}
+	
+	SECTION("Double Attribute")
+	{
+		Tester(PCGAttributeTestsCommonHelper::DoubleTester());
+	}
+	
+	SECTION("Vector Attribute")
+	{
+		SECTION("As Basic Type")
+		{
+			Tester(PCGAttributeTestsCommonHelper::VectorTester_Basic());
+		}
+		
+		SECTION("As Struct")
+		{
+			Tester(PCGAttributeTestsCommonHelper::VectorTester_Struct());
+		}
+	}
+	
+	SECTION("String Attribute")
+	{
+		Tester(PCGAttributeTestsCommonHelper::StringTester());
+	}
+	
+	SECTION("FPCGPoint Attribute")
+	{		
+		Tester(PCGAttributeTestsCommonHelper::FPCGPointTester());
+	}
+	
+	SECTION("Array of Double Attribute")
+	{		
+		Tester(PCGAttributeTestsCommonHelper::ArrayDoubleTester());
+	}
+	
+	SECTION("Array of FPCGPoint Attribute")
+	{	
+		Tester(PCGAttributeTestsCommonHelper::ArrayFPCGPointTester());
+	}
+	
+	SECTION("Set of Double Attribute")
+	{		
+		Tester(PCGAttributeTestsCommonHelper::SetDoubleTester());
+	}
+	
+	SECTION("Map of String -> Double Attribute")
+	{
+		Tester(PCGAttributeTestsCommonHelper::MapStringDoubleTester());
+	}
+	
+	SECTION("Enum attribute")
+	{
+		Tester(PCGAttributeTestsCommonHelper::EnumTester());
+	}
+}
+
+/**
+ * Verifies that the generic accessor supports type broadcast and construction when the caller's type
+ * does not match the attribute's underlying type. A double attribute is created, then written to with
+ * int32 values (requiring construction from int32 -> double). The values are read back as int32
+ * (broadcast double -> int32) and as FString (broadcast double -> string) to confirm the type
+ * conversion path works end-to-end through SetRange_Internal/GetRange_Internal.
+ */
+TEST_CASE_METHOD(PCGTests::FPCGBaseTestWithBasicContext, "PCG::Metadata::Accessors::GenericAttribute::NotExactType", "[PCG][Accessors][GenericAttribute]")
 {
 	const FName DoubleAttributeName = "Double";
 	
@@ -42,21 +199,85 @@ TEST_CASE_METHOD(PCGTests::FPCGBaseTestWithBasicContext, "PCG::Metadata::Accesso
 	REQUIRE_EQUAL(PointArrayData->MutableMetadata()->GetItemCountForChild(), 0);
 	REQUIRE_EQUAL(DoubleAttributeKeys->GetNum(), 5);
 
-	TStaticArray<double, NumPoints> Values{};
-	REQUIRE(DoubleAttributeAccessor->GetRange<double>(Values, 0, *DoubleAttributeKeys));
-
-	REQUIRE(Algo::AllOf(Values, [](const double Value) { return Value == -1.0; }));
-
-	Values = {0, 1, 2, 3, 4};
-	REQUIRE(DoubleAttributeAccessor->SetRange<double>(Values, 0, *DoubleAttributeKeys));
+	TStaticArray<int32, NumPoints> Values{0, 1, 2, 3, 4};
+	REQUIRE(DoubleAttributeAccessor->SetRange<int32>(Values, 0, *DoubleAttributeKeys, EPCGAttributeAccessorFlags::AllowBroadcastAndConstructible));
 	REQUIRE_EQUAL(PointArrayData->MutableMetadata()->GetItemCountForChild(), 5);
 	
+	TStaticArray<int32, NumPoints> ReadValues{};
+	REQUIRE(DoubleAttributeAccessor->GetRange<int32>(ReadValues, 0, *DoubleAttributeKeys, EPCGAttributeAccessorFlags::AllowBroadcastAndConstructible));
+	CHECK_THAT(Values, Catch::Matchers::RangeEquals(ReadValues, [](const auto& a, const auto& b){ return a == b; } ));
+	
+	TStaticArray<FString, NumPoints> ExpectedValuesString{};
+	TStaticArray<FString, NumPoints> ReadValuesString{};
 	for (int32 i = 0; i < NumPoints; ++i)
 	{
-		REQUIRE_EQUAL(GenericAttribute->GetValueFromItemKey<double>(PointArrayDataMetadataEntryRange[i]), i);
+		ExpectedValuesString[i] = PCG::Private::MetadataTraits<double>::ToString(Values[i]);	
 	}
 	
-	TStaticArray<double, NumPoints> ReadValues{};
-	REQUIRE(DoubleAttributeAccessor->GetRange<double>(ReadValues, 0, *DoubleAttributeKeys));
-	CHECK_THAT(Values, Catch::Matchers::RangeEquals(ReadValues, [](const auto& a, const auto& b){ return a == b; } ));
+	REQUIRE(DoubleAttributeAccessor->GetRange<FString>(ReadValuesString, 0, *DoubleAttributeKeys, EPCGAttributeAccessorFlags::AllowBroadcastAndConstructible));
+	CHECK_THAT(ExpectedValuesString, Catch::Matchers::RangeEquals(ReadValuesString, [](const auto& a, const auto& b){ return a == b; } ));
+}
+
+/**
+ * Tests value deduplication for compressed generic attributes. FString attributes use compressed storage
+ * where duplicate values share the same value key. This test writes 6 values containing 3 unique strings,
+ * verifies that only 3 value keys are allocated, then appends 4 more values (2 existing + 2 new) and
+ * verifies the value key count grows to 5. Finally reads back all 10 values via pointer-based GetRange
+ * to confirm correctness. Exercises the AddCompressedValues and FInValuesSubset code paths.
+ */
+TEST_CASE_METHOD(PCGTests::FPCGBaseTestWithBasicContext, "PCG::Metadata::Accessors::GenericAttribute::CompressedData", "[PCG][AttributeGeneric]")
+{
+	constexpr int32 Seed = 42;
+	
+	// Preparing the data
+	constexpr int32 NumPoints = 10;
+	UPCGPointArrayData* PointArrayData = CreateData<UPCGPointArrayData>();
+
+	PointArrayData->SetNumPoints(NumPoints);
+
+	PointArrayData->AllocateProperties(EPCGPointNativeProperties::MetadataEntry);
+	TPCGValueRange<int64> PointArrayDataMetadataEntryRange = PointArrayData->GetMetadataEntryValueRange(/*bAllocate=*/false);
+		
+	FPCGMetadataDomain* Domain = PointArrayData->MutableMetadata()->GetDefaultMetadataDomain();
+		
+	FRandomStream RandomStream(Seed);
+		
+	FPCGMetadataAttributeGeneric* GenericAttribute = nullptr;
+	GenericAttribute = Domain->CreateGenericAttribute<FString>(PCGAttributeTestsCommonHelper::AttributeName, FString{});
+	REQUIRE_NOT_EQUAL(GenericAttribute, nullptr);
+		
+	TUniquePtr<IPCGAttributeAccessor> AttributeAccessor = PCGAttributeAccessorHelpers::CreateAccessor(PointArrayData, FPCGAttributePropertySelector::CreateAttributeSelector(PCGAttributeTestsCommonHelper::AttributeName));
+	TUniquePtr<IPCGAttributeAccessorKeys> AttributeKeys = PCGAttributeAccessorHelpers::CreateKeys(PointArrayData, FPCGAttributePropertySelector::CreateAttributeSelector(PCGAttributeTestsCommonHelper::AttributeName));
+
+	REQUIRE(AttributeAccessor.IsValid());
+	REQUIRE(AttributeKeys.IsValid());
+	
+	// First write 6 values, that are 3 unique values
+	TStaticArray<FString, NumPoints> Values{};
+	TStaticArray<const FString*, NumPoints> ReadValues{};
+	
+	Values[0] = TEXT("Value 0");
+	Values[1] = TEXT("Value 0");
+	Values[2] = TEXT("Value 1");
+	Values[3] = TEXT("Value 2");
+	Values[4] = TEXT("Value 2");
+	Values[5] = TEXT("Value 1");
+	
+	REQUIRE(AttributeAccessor->SetRange<FString>(MakeArrayView(Values.GetData(), 6), 0, *AttributeKeys));
+	REQUIRE_EQUAL(GenericAttribute->GetNumberOfEntries(), 6);
+	REQUIRE_EQUAL(GenericAttribute->GetValueKeyOffsetForChild(), 3);
+	
+	// Then add new values, some existing, some new
+	Values[6] = TEXT("Value 1");
+	Values[7] = TEXT("Value 2");
+	Values[8] = TEXT("Value 3");
+	Values[9] = TEXT("Value 4");
+	
+	REQUIRE(AttributeAccessor->SetRange<FString>(MakeArrayView(Values.GetData() + 6, 4), 6, *AttributeKeys));
+	REQUIRE_EQUAL(GenericAttribute->GetNumberOfEntries(), 10);
+	REQUIRE_EQUAL(GenericAttribute->GetValueKeyOffsetForChild(), 5);
+	
+	//  Read back to verify we have all our values
+	REQUIRE(AttributeAccessor->GetRange<const FString*>(ReadValues, 0, *AttributeKeys));
+	CHECK_THAT(Values, Catch::Matchers::RangeEquals(ReadValues, [](const FString& a, const FString* b){ return a == *b; } ));
 }
\ No newline at end of file
