[PCG] Added node to align points to other points, so they can be placed side by side.

#Synced-CL 52317749

diff --git a/Engine/Plugins/PCG/Source/PCG/Private/Elements/PCGAlignPoints.cpp b/Engine/Plugins/PCG/Source/PCG/Private/Elements/PCGAlignPoints.cpp
new file mode 100644
--- /dev/null
+++ b/Engine/Plugins/PCG/Source/PCG/Private/Elements/PCGAlignPoints.cpp
@@ -0,0 +1,223 @@
+// Copyright Epic Games, Inc. All Rights Reserved.
+
+#include "Elements/PCGAlignPoints.h"
+
+#include "PCGContext.h"
+#include "Data/PCGBasePointData.h"
+#include "Data/PCGSpatialData.h"
+#include "Helpers/PCGPointHelpers.h"
+
+#include UE_INLINE_GENERATED_CPP_BY_NAME(PCGAlignPoints)
+
+#define LOCTEXT_NAMESPACE "PCGAlignPointsElement"
+
+namespace PCGAlignPoints
+{
+	/** Returns the scalar value at AxisIndex (0=X, 1=Y, 2=Z) for the given referential enum. */
+	double GetBoundsValue(const FBox& Bounds, EPCGAlignPointsAxisReferential Referential, int32 AxisIndex)
+	{
+		switch (Referential)
+		{
+			case EPCGAlignPointsAxisReferential::Min: return Bounds.Min[AxisIndex];
+			case EPCGAlignPointsAxisReferential::Max: return Bounds.Max[AxisIndex];
+			case EPCGAlignPointsAxisReferential::Center: return Bounds.GetCenter()[AxisIndex];
+			default: 
+			{
+				checkNoEntry(); 
+				return Bounds.GetCenter()[AxisIndex];
+			}
+		}
+	}
+}
+
+#if WITH_EDITOR
+
+FText UPCGAlignPointsSettings::GetDefaultNodeTitle() const
+{
+	return LOCTEXT("NodeTitle", "Align Points");
+}
+
+FText UPCGAlignPointsSettings::GetNodeTooltipText() const
+{
+	return LOCTEXT("NodeTooltip", "For each source point, aligns (changes location and optionally rotation) it so that the chosen edge of its bounding box aligns with the chosen edge of the corresponding target point's bounding box. Bounds are computed in the chosen spatial referential (Target, Source, or World axes).");
+}
+
+#endif // WITH_EDITOR
+
+TArray<FPCGPinProperties> UPCGAlignPointsSettings::InputPinProperties() const
+{
+	TArray<FPCGPinProperties> PinProperties;
+
+	FPCGPinProperties& SourcePin = PinProperties.Emplace_GetRef(PCGAlignPointsConstants::SourcePointsLabel, EPCGDataType::Point);
+#if WITH_EDITOR
+	SourcePin.Tooltip = LOCTEXT("SourcePinTooltip", "The points to align. Each source point data entry is paired with a target point data entry.");
+#endif
+	SourcePin.SetRequiredPin();
+
+	FPCGPinProperties& TargetPin = PinProperties.Emplace_GetRef(PCGAlignPointsConstants::TargetPointsLabel, EPCGDataType::Point);
+#if WITH_EDITOR
+	TargetPin.Tooltip = LOCTEXT("TargetPinTooltip", "The reference points to align to. Must have the same number of data entries as Source, or a single entry applied to all sources. Within each pair, target must have the same number of points as source, or a single point.");
+#endif
+	TargetPin.SetRequiredPin();
+
+	return PinProperties;
+}
+
+FPCGElementPtr UPCGAlignPointsSettings::CreateElement() const
+{
+	return MakeShared<FPCGAlignPointsElement>();
+}
+
+bool FPCGAlignPointsElement::ExecuteInternal(FPCGContext* Context) const
+{
+	TRACE_CPUPROFILER_EVENT_SCOPE(FPCGAlignPointsElement::Execute);
+
+	const UPCGAlignPointsSettings* Settings = Context->GetInputSettings<UPCGAlignPointsSettings>();
+	check(Settings);
+
+	const TArray<FPCGTaggedData> Sources = Context->InputData.GetInputsByPin(PCGAlignPointsConstants::SourcePointsLabel);
+	const TArray<FPCGTaggedData> Targets = Context->InputData.GetInputsByPin(PCGAlignPointsConstants::TargetPointsLabel);
+	TArray<FPCGTaggedData>& Outputs = Context->OutputData.TaggedData;
+
+	const int32 NumSources = Sources.Num();
+	const int32 NumTargets = Targets.Num();
+
+	if (NumSources == 0 || NumTargets == 0)
+	{
+		Outputs = Sources;
+		return true;
+	}
+
+	// Validate data-level pairing: must be N:N, or N:1.
+	if (NumSources != NumTargets && NumTargets != 1)
+	{
+		PCGLog::InputOutput::LogInvalidCardinalityError(PCGAlignPointsConstants::SourcePointsLabel, PCGAlignPointsConstants::TargetPointsLabel, Context);
+		Outputs = Sources;
+		return true;
+	}
+
+	const EPCGAlignPointsSpatialReferential SpatialReferential = Settings->SpatialReferential;
+	const FPCGAlignPointsAxisSettings& XAxisSettings = Settings->XAxis;
+	const FPCGAlignPointsAxisSettings& YAxisSettings = Settings->YAxis;
+	const FPCGAlignPointsAxisSettings& ZAxisSettings = Settings->ZAxis;
+
+	// Early-out: no axis enabled means no work to do — pass sources through unchanged.
+	if (!XAxisSettings.bApplyToAxis && !YAxisSettings.bApplyToAxis && !ZAxisSettings.bApplyToAxis)
+	{
+		Outputs = Sources;
+		return true;
+	}
+
+	for (int32 SourceIndex = 0; SourceIndex < NumSources; ++SourceIndex)
+	{
+		const int32 TargetIndex = (SourceIndex % NumTargets);
+		const FPCGTaggedData& SourceTaggedData = Sources[SourceIndex];
+		const FPCGTaggedData& TargetTaggedData = Targets[TargetIndex];
+
+		// Start output as a copy of the source tagged data (preserves tags, pin label, etc.)
+		FPCGTaggedData& Output = Outputs.Add_GetRef(SourceTaggedData);
+
+		if (!SourceTaggedData.Data || !TargetTaggedData.Data)
+		{
+			PCGLog::InputOutput::LogInvalidInputDataError(Context);
+			continue;
+		}
+
+		const UPCGBasePointData* SourcePointData = Cast<UPCGBasePointData>(SourceTaggedData.Data);
+		const UPCGBasePointData* TargetPointData = Cast<UPCGBasePointData>(TargetTaggedData.Data);
+
+		if (!SourcePointData || !TargetPointData)
+		{
+			continue;
+		}
+
+		const int32 NumSourcePoints = SourcePointData->GetNumPoints();
+		const int32 NumTargetPoints = TargetPointData->GetNumPoints();
+
+		// Validate point-level pairing within this data pair.
+		if (NumTargetPoints != NumSourcePoints && NumTargetPoints != 1)
+		{
+			PCGLog::LogErrorOnGraph(FText::Format(LOCTEXT("NumPointsMismatch", "Source point count ({0}) does not match Target point count ({1}) for data pair {2}. Target must have the same number of points as Source, or exactly one point."), NumSourcePoints, NumTargetPoints, SourceIndex));
+			continue;
+		}
+
+		// Build the output point data as a copy of the source, with modified transforms.
+		UPCGBasePointData* OutPointData = CastChecked<UPCGBasePointData>(SourcePointData->DuplicateData(Context));
+		Output.Data = OutPointData;
+
+		OutPointData->AllocateProperties(SourcePointData->GetAllocatedProperties() | EPCGPointNativeProperties::Transform);
+
+		// Read source value ranges.
+		FConstPCGPointValueRanges SourceRanges(SourcePointData);
+		FConstPCGPointValueRanges TargetRanges(TargetPointData);
+		TPCGValueRange<FTransform> OutTransformRange = OutPointData->GetTransformValueRange();
+
+		for (int32 PointIndex = 0; PointIndex < NumSourcePoints; ++PointIndex)
+		{
+			// Then retrieve both the source & target points so we can compute the appropriate delta to apply.
+			FPCGPoint SourcePoint = SourceRanges.GetPoint(PointIndex);
+
+			FTransform SourcePointTransform = SourcePoint.Transform;
+
+			const int32 TargetPointIndex = (NumTargetPoints == 1) ? 0 : PointIndex;
+			FPCGPoint TargetPoint = TargetRanges.GetPoint(TargetPointIndex);
+
+			FBox SourceBounds = SourcePoint.GetLocalBounds();
+			FBox TargetBounds = TargetPoint.GetLocalBounds();
+			FTransform RelativeTransform = FTransform::Identity;
+			FTransform ReferentialTransform = FTransform::Identity;
+
+			// Local to target mode enforces the source points to align to the target, prior to alignment.
+			if (SpatialReferential == EPCGAlignPointsSpatialReferential::LocalToTarget)
+			{
+				SourcePointTransform.SetRotation(TargetPoint.Transform.GetRotation());
+			}
+
+			if (SpatialReferential == EPCGAlignPointsSpatialReferential::Target || SpatialReferential == EPCGAlignPointsSpatialReferential::LocalToTarget)
+			{
+				ReferentialTransform = TargetPoint.Transform;
+				RelativeTransform = SourcePointTransform.GetRelativeTransform(TargetPoint.Transform);
+				SourceBounds = SourceBounds.TransformBy(RelativeTransform);
+			}
+			else if (SpatialReferential == EPCGAlignPointsSpatialReferential::Source)
+			{
+				ReferentialTransform = SourcePointTransform;
+				RelativeTransform = TargetPoint.Transform.GetRelativeTransform(SourcePoint.Transform);
+				TargetBounds = TargetBounds.TransformBy(RelativeTransform);
+			}
+			else // World
+			{
+				SourceBounds = SourceBounds.TransformBy(SourcePoint.Transform);
+				TargetBounds = TargetBounds.TransformBy(TargetPoint.Transform);
+			}
+
+			// Compute per-axis alignment delta in the referential frame.
+			FVector DeltaInRef = FVector::ZeroVector;
+			if (XAxisSettings.bApplyToAxis)
+			{
+				DeltaInRef.X = PCGAlignPoints::GetBoundsValue(TargetBounds, XAxisSettings.TargetReferential, 0) - PCGAlignPoints::GetBoundsValue(SourceBounds, XAxisSettings.SourceReferential, 0);
+			}
+
+			if (YAxisSettings.bApplyToAxis)
+			{
+				DeltaInRef.Y = PCGAlignPoints::GetBoundsValue(TargetBounds, YAxisSettings.TargetReferential, 1) - PCGAlignPoints::GetBoundsValue(SourceBounds, YAxisSettings.SourceReferential, 1);
+			}
+
+			if (ZAxisSettings.bApplyToAxis)
+			{
+				DeltaInRef.Z = PCGAlignPoints::GetBoundsValue(TargetBounds, ZAxisSettings.TargetReferential, 2) - PCGAlignPoints::GetBoundsValue(SourceBounds, ZAxisSettings.SourceReferential, 2);
+			}
+
+			// Convert the delta from the referential frame back to world space
+			const FVector WorldDelta = ReferentialTransform.TransformVector(DeltaInRef);
+
+			// Write the aligned transform: same rotation and scale as source, translated position.
+			SourcePointTransform.AddToTranslation(WorldDelta);
+			OutTransformRange[PointIndex] = SourcePointTransform;
+		}
+	}
+
+	return true;
+}
+
+#undef LOCTEXT_NAMESPACE
\ No newline at end of file
diff --git a/Engine/Plugins/PCG/Source/PCG/Public/Elements/PCGAlignPoints.h b/Engine/Plugins/PCG/Source/PCG/Public/Elements/PCGAlignPoints.h
new file mode 100644
--- /dev/null
+++ b/Engine/Plugins/PCG/Source/PCG/Public/Elements/PCGAlignPoints.h
@@ -0,0 +1,106 @@
+// Copyright Epic Games, Inc. All Rights Reserved.
+
+#pragma once
+
+#include "PCGSettings.h"
+
+#include "PCGAlignPoints.generated.h"
+
+namespace PCGAlignPointsConstants
+{
+	const FName SourcePointsLabel = TEXT("Source");
+	const FName TargetPointsLabel = TEXT("Target");
+}
+
+/** Controls which edge of a bounding box is used as the reference value on a given axis. */
+UENUM(BlueprintType)
+enum class EPCGAlignPointsAxisReferential : uint8
+{
+	Min,
+	Max,
+	Center
+};
+
+/** Defines which coordinate frame is used to compute bounding box extents for alignment. */
+UENUM(BlueprintType)
+enum class EPCGAlignPointsSpatialReferential : uint8
+{
+	/** Express bounds along the target point's local axes. */
+	Target,
+	/** Align source to target's rotation, then affect using the target referential. */
+	LocalToTarget,
+	/** Express bounds along the source point's local axes. */
+	Source,
+	/** Express bounds along world-space axes (standard AABB). */
+	World
+};
+
+/** Settings controlling alignment along a single axis. */
+USTRUCT(BlueprintType)
+struct FPCGAlignPointsAxisSettings
+{
+	GENERATED_BODY()
+
+	/** Whether to apply alignment on this axis. */
+	UPROPERTY(BlueprintReadWrite, EditAnywhere, Category = "")
+	bool bApplyToAxis = false;
+
+	/** Which edge of the source point's bounding box to align from. */
+	UPROPERTY(BlueprintReadWrite, EditAnywhere, Category = "", meta = (EditCondition = "bApplyToAxis", EditConditionHides))
+	EPCGAlignPointsAxisReferential SourceReferential = EPCGAlignPointsAxisReferential::Center;
+
+	/** Which edge of the target point's bounding box to align to. */
+	UPROPERTY(BlueprintReadWrite, EditAnywhere, Category = "", meta = (EditCondition = "bApplyToAxis", EditConditionHides))
+	EPCGAlignPointsAxisReferential TargetReferential = EPCGAlignPointsAxisReferential::Center;
+};
+
+UCLASS(MinimalAPI, BlueprintType, ClassGroup = (Procedural))
+class UPCGAlignPointsSettings : public UPCGSettings
+{
+	GENERATED_BODY()
+
+public:
+	//~Begin UPCGSettings interface
+#if WITH_EDITOR
+	virtual FName GetDefaultNodeName() const override { return FName(TEXT("AlignPoints")); }
+	virtual FText GetDefaultNodeTitle() const override;
+	virtual FText GetNodeTooltipText() const override;
+	virtual EPCGSettingsType GetType() const override { return EPCGSettingsType::PointOps; }
+#endif
+
+protected:
+	virtual TArray<FPCGPinProperties> InputPinProperties() const override;
+	virtual TArray<FPCGPinProperties> OutputPinProperties() const override { return Super::DefaultPointOutputPinProperties(); }
+	virtual FPCGElementPtr CreateElement() const override;
+	//~End UPCGSettings interface
+
+public:
+	/**
+	 * Coordinate frame used to compute bounding box extents when determining alignment values.
+	 * Target: bounds are measured along the target point's local axes.
+	 * LocalToTarget: source point is rotated like the target point, then bounds are measured along the tàrget point's local axes.
+	 * Source: bounds are measured along the source point's own local axes.
+	 * World:  bounds are measured along world-aligned axes (standard AABB).
+	 */
+	UPROPERTY(BlueprintReadWrite, EditAnywhere, Category = Settings)
+	EPCGAlignPointsSpatialReferential SpatialReferential = EPCGAlignPointsSpatialReferential::World;
+
+	/** Alignment settings along the X axis of the chosen spatial referential. */
+	UPROPERTY(BlueprintReadWrite, EditAnywhere, Category = "Settings", meta = (PCG_Overridable))
+	FPCGAlignPointsAxisSettings XAxis;
+
+	/** Alignment settings along the Y axis of the chosen spatial referential. */
+	UPROPERTY(BlueprintReadWrite, EditAnywhere, Category = "Settings", meta = (PCG_Overridable))
+	FPCGAlignPointsAxisSettings YAxis;
+
+	/** Alignment settings along the Z axis of the chosen spatial referential. */
+	UPROPERTY(BlueprintReadWrite, EditAnywhere, Category = "Settings", meta = (PCG_Overridable))
+	FPCGAlignPointsAxisSettings ZAxis;
+};
+
+class FPCGAlignPointsElement : public IPCGElement
+{
+protected:
+	virtual bool ExecuteInternal(FPCGContext* Context) const override;
+	virtual bool SupportsBasePointDataInputs(FPCGContext* InContext) const override { return true; }
+};
diff --git a/Engine/Plugins/PCG/Tests/Private/Elements/PCGAlignPointsTest.cpp b/Engine/Plugins/PCG/Tests/Private/Elements/PCGAlignPointsTest.cpp
new file mode 100644
--- /dev/null
+++ b/Engine/Plugins/PCG/Tests/Private/Elements/PCGAlignPointsTest.cpp
@@ -0,0 +1,454 @@
+// Copyright Epic Games, Inc. All Rights Reserved.
+
+#include "PCGTestsCommon.h"
+
+#include "Data/PCGBasePointData.h"
+#include "Data/PCGPointArrayData.h"
+#include "Elements/PCGAlignPoints.h"
+
+#include <catch2/catch_test_macros.hpp>
+
+#include "TestHarness.h"
+
+namespace
+{
+	const FVector SrcBoundsMin(-50.f, -40.f, -30.f);
+	const FVector SrcBoundsMax( 50.f,  40.f,  30.f);
+	const FVector TgtBoundsMin(-30.f, -20.f, -10.f);
+	const FVector TgtBoundsMax( 60.f,  70.f,  80.f);
+	const FVector SrcPos( 100.f, 200.f,  50.f);
+	const FVector TgtPos( -50.f, 100.f, 200.f);
+	constexpr float PosTolerance = 1.0e-2f;
+} // anonymous namespace
+
+// ---------------------------------------------------------------------------
+// Test geometry (chosen so expected values can be hand-verified):
+//
+//   Source local bounds:  Min=(-50,-40,-30)  Max=(50,40,30)   LocalCenter=(0,0,0)
+//   Target local bounds:  Min=(-30,-20,-10)  Max=(60,70,80)   LocalCenter=(15,25,35)
+//   Source position:  (100, 200, 50)
+//   Target position:  (-50, 100, 200)
+//
+// World-space AABBs for identity transforms:
+//   Src world: Min=(50,160,20)   Max=(150,240,80)  Center=(100,200,50)
+//   Tgt world: Min=(-80,80,190)  Max=(10,170,280)  Center=(-35,125,235)
+// ---------------------------------------------------------------------------
+
+class FPCGAlignPointsBaseTest : public PCGTests::FPCGSingleElementBaseTest<UPCGAlignPointsSettings>
+{
+public:
+	FPCGAlignPointsBaseTest() : FPCGSingleElementBaseTest() {}
+
+protected:
+	using ERef     = EPCGAlignPointsAxisReferential;
+	using ESpatial = EPCGAlignPointsSpatialReferential;
+
+	UPCGPointArrayData* MakePointData(const FTransform& Transform, const FVector& BoundsMin, const FVector& BoundsMax)
+	{
+		UPCGPointArrayData* PD = NewObject<UPCGPointArrayData>();
+		PD->SetNumPoints(1);
+		FPCGPointValueRanges Ranges(PD);
+		Ranges.TransformRange[0] = Transform;
+		Ranges.BoundsMinRange[0] = BoundsMin;
+		Ranges.BoundsMaxRange[0] = BoundsMax;
+		return PD;
+	}
+
+	UPCGPointArrayData* MakePointData(const FVector& Position, const FVector& BoundsMin, const FVector& BoundsMax)
+	{
+		return MakePointData(FTransform(Position), BoundsMin, BoundsMax);
+	}
+
+	FPCGTaggedData MakeSourceEntry(UPCGPointArrayData* Data)
+	{
+		FPCGTaggedData Entry;
+		Entry.Data = Data;
+		Entry.Pin = PCGAlignPointsConstants::SourcePointsLabel;
+		return Entry;
+	}
+
+	FPCGTaggedData MakeTargetEntry(UPCGPointArrayData* Data)
+	{
+		FPCGTaggedData Entry;
+		Entry.Data = Data;
+		Entry.Pin = PCGAlignPointsConstants::TargetPointsLabel;
+		return Entry;
+	}
+
+	void SetAxis(bool bX, ERef SrcX, ERef TgtX,
+	             bool bY, ERef SrcY, ERef TgtY,
+	             bool bZ, ERef SrcZ, ERef TgtZ)
+	{
+		TypedSettings->XAxis = { bX, SrcX, TgtX };
+		TypedSettings->YAxis = { bY, SrcY, TgtY };
+		TypedSettings->ZAxis = { bZ, SrcZ, TgtZ };
+	}
+
+	TArray<FTransform> RunAlignPointsTransforms(const TArray<FPCGTaggedData>& SourceEntries, const TArray<FPCGTaggedData>& TargetEntries)
+	{
+		for (const FPCGTaggedData& Entry : SourceEntries)
+		{
+			InputData.TaggedData.Add(Entry);
+		}
+
+		for (const FPCGTaggedData& Entry : TargetEntries)
+		{
+			InputData.TaggedData.Add(Entry);
+		}
+
+		ExecuteElement();
+
+		TArray<FTransform> OutTransforms;
+		for (const FPCGTaggedData& Output : Context->OutputData.TaggedData)
+		{
+			const UPCGBasePointData* PointData = Cast<UPCGBasePointData>(Output.Data);
+			if (!PointData)
+			{
+				continue;
+			}
+
+			TConstPCGValueRange<FTransform> TransformRange = PointData->GetConstTransformValueRange();
+			for (int32 i = 0; i < PointData->GetNumPoints(); ++i)
+			{
+				OutTransforms.Add(TransformRange[i]);
+			}
+		}
+		return OutTransforms;
+	}
+
+	TArray<FVector> RunAlignPoints(const TArray<FPCGTaggedData>& SourceEntries, const TArray<FPCGTaggedData>& TargetEntries)
+	{
+		TArray<FVector> OutPositions;
+		for (const FTransform& T : RunAlignPointsTransforms(SourceEntries, TargetEntries))
+		{
+			OutPositions.Add(T.GetLocation());
+		}
+
+		return OutPositions;
+	}
+
+	FVector Run1v1(const FTransform& Src, const FTransform& Tgt)
+	{
+		const TArray<FVector> Out = RunAlignPoints(
+			{ MakeSourceEntry(MakePointData(Src, SrcBoundsMin, SrcBoundsMax)) },
+			{ MakeTargetEntry(MakePointData(Tgt, TgtBoundsMin, TgtBoundsMax)) });
+		return (Out.Num() == 1) ? Out[0] : FVector(UE_BIG_NUMBER);
+	}
+
+	FTransform Run1v1T(const FTransform& Src, const FTransform& Tgt)
+	{
+		const TArray<FTransform> Out = RunAlignPointsTransforms(
+			{ MakeSourceEntry(MakePointData(Src, SrcBoundsMin, SrcBoundsMax)) },
+			{ MakeTargetEntry(MakePointData(Tgt, TgtBoundsMin, TgtBoundsMax)) });
+		return (Out.Num() == 1) ? Out[0] : FTransform::Identity;
+	}
+};
+
+// =========================================================================
+// World referential — identity transforms, all expected values hand-computed
+//
+// Src world AABB: Min=(50,160,20)  Max=(150,240,80)  Center=(100,200,50)
+// Tgt world AABB: Min=(-80,80,190) Max=(10,170,280)  Center=(-35,125,235)
+// =========================================================================
+
+TEST_CASE_METHOD(FPCGAlignPointsBaseTest, "PCG::AlignPoints::World", "[PCG][AlignPoints]")
+{
+	TypedSettings->SpatialReferential = ESpatial::World;
+	const FTransform SrcIdentity(FQuat::Identity, SrcPos, FVector::OneVector);
+	const FTransform TgtIdentity(FQuat::Identity, TgtPos, FVector::OneVector);
+
+	SECTION("NoOp")
+	{
+		// No axes enabled: source must pass through unchanged.
+		SetAxis(false, ERef::Center, ERef::Center,
+		        false, ERef::Center, ERef::Center,
+		        false, ERef::Center, ERef::Center);
+		const FVector Out = Run1v1(SrcIdentity, TgtIdentity);
+		CHECK(Out.Equals(SrcPos, PosTolerance));
+	}
+
+	SECTION("AllCenter")
+	{
+		// Delta = TgtWorldCenter - SrcWorldCenter = (-135,-75,185)
+		// OutPos = (100,200,50) + (-135,-75,185) = (-35,125,235)
+		SetAxis(true, ERef::Center, ERef::Center,
+		        true, ERef::Center, ERef::Center,
+		        true, ERef::Center, ERef::Center);
+		const FVector Out = Run1v1(SrcIdentity, TgtIdentity);
+		CHECK(Out.Equals(FVector(-35.f, 125.f, 235.f), PosTolerance));
+	}
+
+	SECTION("AllMin")
+	{
+		// Delta = TgtWorldMin - SrcWorldMin = (-80-50, 80-160, 190-20) = (-130,-80,170)
+		// OutPos = (100,200,50) + (-130,-80,170) = (-30,120,220)
+		SetAxis(true, ERef::Min, ERef::Min,
+		        true, ERef::Min, ERef::Min,
+		        true, ERef::Min, ERef::Min);
+		const FVector Out = Run1v1(SrcIdentity, TgtIdentity);
+		CHECK(Out.Equals(FVector(-30.f, 120.f, 220.f), PosTolerance));
+	}
+
+	SECTION("AllMax")
+	{
+		// Delta = TgtWorldMax - SrcWorldMax = (10-150, 170-240, 280-80) = (-140,-70,200)
+		// OutPos = (100,200,50) + (-140,-70,200) = (-40,130,250)
+		SetAxis(true, ERef::Max, ERef::Max,
+		        true, ERef::Max, ERef::Max,
+		        true, ERef::Max, ERef::Max);
+		const FVector Out = Run1v1(SrcIdentity, TgtIdentity);
+		CHECK(Out.Equals(FVector(-40.f, 130.f, 250.f), PosTolerance));
+	}
+
+	SECTION("XOnly")
+	{
+		// Delta.X = TgtWorldCenter.X - SrcWorldCenter.X = -35 - 100 = -135
+		// OutPos = (-35, 200, 50)
+		SetAxis(true,  ERef::Center, ERef::Center,
+		        false, ERef::Center, ERef::Center,
+		        false, ERef::Center, ERef::Center);
+		const FVector Out = Run1v1(SrcIdentity, TgtIdentity);
+		CHECK(Out.Equals(FVector(-35.f, 200.f, 50.f), PosTolerance));
+	}
+
+	SECTION("ZOnly")
+	{
+		// Delta.Z = TgtWorldMin.Z - SrcWorldMin.Z = 190 - 20 = 170
+		// OutPos = (100, 200, 220)
+		SetAxis(false, ERef::Center, ERef::Center,
+		        false, ERef::Center, ERef::Center,
+		        true,  ERef::Min, ERef::Min);
+		const FVector Out = Run1v1(SrcIdentity, TgtIdentity);
+		CHECK(Out.Equals(FVector(100.f, 200.f, 220.f), PosTolerance));
+	}
+
+	SECTION("YMaxToMin")
+	{
+		// SrcWorldMax.Y = 240, TgtWorldMin.Y = 80
+		// Delta.Y = 80 - 240 = -160 → OutPos = (100, 40, 50)
+		SetAxis(false, ERef::Center, ERef::Center,
+		        true,  ERef::Max, ERef::Min,
+		        false, ERef::Center, ERef::Center);
+		const FVector Out = Run1v1(SrcIdentity, TgtIdentity);
+		CHECK(Out.Equals(FVector(100.f, 40.f, 50.f), PosTolerance));
+	}
+
+	SECTION("Mixed")
+	{
+		// X: center-to-min  → Delta.X = TgtWorldMin.X - SrcWorldCenter.X = -80 - 100 = -180
+		// Y: min-to-max     → Delta.Y = TgtWorldMax.Y - SrcWorldMin.Y   = 170 - 160 = 10
+		// Z: max-to-center  → Delta.Z = TgtWorldCenter.Z - SrcWorldMax.Z = 235 - 80  = 155
+		// OutPos = (100-180, 200+10, 50+155) = (-80, 210, 205)
+		SetAxis(true, ERef::Center, ERef::Min,
+		        true, ERef::Min,    ERef::Max,
+		        true, ERef::Max,    ERef::Center);
+		const FVector Out = Run1v1(SrcIdentity, TgtIdentity);
+		CHECK(Out.Equals(FVector(-80.f, 210.f, 205.f), PosTolerance));
+	}
+}
+
+// =========================================================================
+// Target referential — target has 90° yaw, source is identity
+//
+// In UE: yaw=90 maps +X→+Y, +Y→-X.
+// After computing source bounds in target-local space, center-to-center gives
+// OutPos = (-75, 115, 235).
+// =========================================================================
+
+TEST_CASE_METHOD(FPCGAlignPointsBaseTest, "PCG::AlignPoints::Target", "[PCG][AlignPoints]")
+{
+	TypedSettings->SpatialReferential = ESpatial::Target;
+	const FTransform SrcIdentity(FQuat::Identity, SrcPos, FVector::OneVector);
+
+	SECTION("AllCenter_TgtYaw90")
+	{
+		const FTransform TgtYaw90(FRotator(0.f, 90.f, 0.f), TgtPos, FVector::OneVector);
+		SetAxis(true, ERef::Center, ERef::Center,
+		        true, ERef::Center, ERef::Center,
+		        true, ERef::Center, ERef::Center);
+		const FVector Out = Run1v1(SrcIdentity, TgtYaw90);
+		CHECK(Out.Equals(FVector(-75.f, 115.f, 235.f), PosTolerance));
+	}
+}
+
+// =========================================================================
+// Source referential — source has 90° yaw, target is identity
+//
+// Center-to-center DeltaInRef=(-75,135,185), WorldDelta=yaw90*(-75,135,185)=(-135,-75,185)
+// OutPos = (-35, 125, 235).
+//
+// Min-to-min: TargetBounds in source-local Min=(-120,90,140),
+// DeltaInRef=(-70,130,170), WorldDelta=yaw90*(-70,130,170)=(-130,-70,170)
+// OutPos = (-30, 130, 220).
+// =========================================================================
+
+TEST_CASE_METHOD(FPCGAlignPointsBaseTest, "PCG::AlignPoints::Source", "[PCG][AlignPoints]")
+{
+	TypedSettings->SpatialReferential = ESpatial::Source;
+	const FTransform TgtIdentity(FQuat::Identity, TgtPos, FVector::OneVector);
+	const FTransform SrcYaw90(FRotator(0.f, 90.f, 0.f), SrcPos, FVector::OneVector);
+
+	SECTION("AllCenter_SrcYaw90")
+	{
+		SetAxis(true, ERef::Center, ERef::Center,
+		        true, ERef::Center, ERef::Center,
+		        true, ERef::Center, ERef::Center);
+		const FVector Out = Run1v1(SrcYaw90, TgtIdentity);
+		CHECK(Out.Equals(FVector(-35.f, 125.f, 235.f), PosTolerance));
+	}
+
+	SECTION("AllMin_SrcYaw90")
+	{
+		SetAxis(true, ERef::Min, ERef::Min,
+		        true, ERef::Min, ERef::Min,
+		        true, ERef::Min, ERef::Min);
+		const FVector Out = Run1v1(SrcYaw90, TgtIdentity);
+		CHECK(Out.Equals(FVector(-30.f, 130.f, 220.f), PosTolerance));
+	}
+}
+
+// =========================================================================
+// Pairing: N:1 and N:N
+// World referential, all axes, center-to-center.
+// =========================================================================
+
+TEST_CASE_METHOD(FPCGAlignPointsBaseTest, "PCG::AlignPoints::Pairing", "[PCG][AlignPoints]")
+{
+	TypedSettings->SpatialReferential = ESpatial::World;
+	SetAxis(true, ERef::Center, ERef::Center,
+	        true, ERef::Center, ERef::Center,
+	        true, ERef::Center, ERef::Center);
+
+	const FTransform SrcIdentity(FQuat::Identity, SrcPos, FVector::OneVector);
+	const FTransform TgtIdentity(FQuat::Identity, TgtPos, FVector::OneVector);
+
+	SECTION("N_to_1")
+	{
+		// Two source points both aligned to the single target center (-35,125,235).
+		// Src0=(100,200,50), Src1=(0,0,0), both with same bounds.
+		UPCGPointArrayData* Src2Pts = NewObject<UPCGPointArrayData>();
+		Src2Pts->SetNumPoints(2);
+		{
+			FPCGPointValueRanges R(Src2Pts);
+			R.TransformRange[0] = SrcIdentity;
+			R.BoundsMinRange[0] = SrcBoundsMin;
+			R.BoundsMaxRange[0] = SrcBoundsMax;
+			R.TransformRange[1] = FTransform::Identity;
+			R.BoundsMinRange[1] = SrcBoundsMin;
+			R.BoundsMaxRange[1] = SrcBoundsMax;
+		}
+
+		FPCGTaggedData SrcEntry;
+		SrcEntry.Data = Src2Pts;
+		SrcEntry.Pin  = PCGAlignPointsConstants::SourcePointsLabel;
+
+		const TArray<FVector> Out = RunAlignPoints(
+			{ SrcEntry },
+			{ MakeTargetEntry(MakePointData(TgtIdentity, TgtBoundsMin, TgtBoundsMax)) });
+
+		REQUIRE_EQUAL(Out.Num(), 2);
+		CHECK(Out[0].Equals(FVector(-35.f, 125.f, 235.f), PosTolerance));
+		CHECK(Out[1].Equals(FVector(-35.f, 125.f, 235.f), PosTolerance));
+	}
+
+	SECTION("N_to_N")
+	{
+		// Entry 0: Src=(100,200,50)  Tgt=(-50,100,200)  → OutPos=(-35,125,235)
+		// Entry 1: Src=(0,0,0)       Tgt=(300,300,300)
+		//   Tgt1 world center: (315,325,335) = (300+15, 300+25, 300+35)
+		//   OutPos1 = (315,325,335)
+		const FVector Src1Pos(0.f, 0.f, 0.f);
+		const FVector Tgt1Pos(300.f, 300.f, 300.f);
+
+		const TArray<FVector> Out = RunAlignPoints(
+			{ MakeSourceEntry(MakePointData(SrcPos,   SrcBoundsMin, SrcBoundsMax)),
+			  MakeSourceEntry(MakePointData(Src1Pos,  SrcBoundsMin, SrcBoundsMax)) },
+			{ MakeTargetEntry(MakePointData(TgtPos,   TgtBoundsMin, TgtBoundsMax)),
+			  MakeTargetEntry(MakePointData(Tgt1Pos,  TgtBoundsMin, TgtBoundsMax)) });
+
+		REQUIRE_EQUAL(Out.Num(), 2);
+		CHECK(Out[0].Equals(FVector(-35.f, 125.f, 235.f), PosTolerance));
+		CHECK(Out[1].Equals(FVector(315.f, 325.f, 335.f), PosTolerance));
+	}
+}
+
+// =========================================================================
+// LocalToTarget referential
+//
+// Behavior: before computing bounds, the source's rotation is replaced by
+// the target's rotation. The RelativeTransform then has Identity rotation,
+// so source bounds are translated (not rotated) into target-local space.
+// The output transform inherits the target's rotation (not the source's).
+// =========================================================================
+
+TEST_CASE_METHOD(FPCGAlignPointsBaseTest, "PCG::AlignPoints::LocalToTarget", "[PCG][AlignPoints]")
+{
+	TypedSettings->SpatialReferential = ESpatial::LocalToTarget;
+	SetAxis(true, ERef::Center, ERef::Center,
+	        true, ERef::Center, ERef::Center,
+	        true, ERef::Center, ERef::Center);
+
+	const FTransform SrcIdentity(FQuat::Identity, SrcPos, FVector::OneVector);
+	const FTransform TgtIdentity(FQuat::Identity, TgtPos, FVector::OneVector);
+	const FQuat Yaw90 = FQuat(FRotator(0.f, 90.f, 0.f));
+
+	SECTION("BothIdentity")
+	{
+		// SourcePointTransform rotation → Identity (= TgtRot).
+		// RelativeTransform = (Identity, (150,100,-150), 1).
+		// Source bounds center in target-local = (150,100,-150). Target center = (15,25,35).
+		// DeltaInRef = (-135,-75,185). WorldDelta = (-135,-75,185).
+		// OutPos = (-35,125,235). OutRot = Identity.
+		const FTransform Out = Run1v1T(SrcIdentity, TgtIdentity);
+		CHECK(Out.GetLocation().Equals(FVector(-35.f, 125.f, 235.f), PosTolerance));
+		CHECK(Out.GetRotation().Equals(FQuat::Identity, PosTolerance));
+	}
+
+	SECTION("SrcYaw90_TgtIdentity")
+	{
+		// Source rotation overridden to Identity (= TgtRot).
+		// Same relative translation as BothIdentity → OutPos = (-35,125,235). OutRot = Identity.
+		const FTransform SrcYaw90(Yaw90, SrcPos, FVector::OneVector);
+		const FTransform Out = Run1v1T(SrcYaw90, TgtIdentity);
+		CHECK(Out.GetLocation().Equals(FVector(-35.f, 125.f, 235.f), PosTolerance));
+		CHECK(Out.GetRotation().Equals(FQuat::Identity, PosTolerance));
+	}
+
+	SECTION("SrcIdentity_TgtYaw90")
+	{
+		// Source rotation overridden to yaw90. RelativeTransform rotation = Identity.
+		// RelTranslation in tgt-local = yaw=-90 * (SrcPos - TgtPos) = yaw=-90*(150,100,-150) = (100,-150,-150).
+		// Source bounds center in tgt-local = (100,-150,-150). Target center = (15,25,35).
+		// DeltaInRef = (-85,175,185). WorldDelta = yaw90*(-85,175,185) = (-175,-85,185).
+		// OutPos = (-75,115,235). OutRot = yaw90.
+		const FTransform TgtYaw90(Yaw90, TgtPos, FVector::OneVector);
+		const FTransform Out = Run1v1T(SrcIdentity, TgtYaw90);
+		CHECK(Out.GetLocation().Equals(FVector(-75.f, 115.f, 235.f), PosTolerance));
+		CHECK(Out.GetRotation().Equals(Yaw90, PosTolerance));
+	}
+
+	SECTION("BothYaw90")
+	{
+		// Source rotation already = TgtRot → behaves like Target/AllCenter/TgtYaw90.
+		// OutPos = (-75,115,235). OutRot = yaw90.
+		const FTransform SrcYaw90(Yaw90, SrcPos, FVector::OneVector);
+		const FTransform TgtYaw90(Yaw90, TgtPos, FVector::OneVector);
+		const FTransform Out = Run1v1T(SrcYaw90, TgtYaw90);
+		CHECK(Out.GetLocation().Equals(FVector(-75.f, 115.f, 235.f), PosTolerance));
+		CHECK(Out.GetRotation().Equals(Yaw90, PosTolerance));
+	}
+
+	SECTION("AllMin_SrcYaw90_TgtIdentity")
+	{
+		// Source bounds min in tgt-local (identity relative rotation, translated by (150,100,-150)):
+		//   Min = SrcBoundsMin + (150,100,-150) = (100,60,-180). TargetBoundsMin = (-30,-20,-10).
+		// DeltaInRef = (-130,-80,170). WorldDelta = (-130,-80,170).
+		// OutPos = (-30,120,220).
+		const FTransform SrcYaw90(Yaw90, SrcPos, FVector::OneVector);
+		SetAxis(true, ERef::Min, ERef::Min,
+		        true, ERef::Min, ERef::Min,
+		        true, ERef::Min, ERef::Min);
+		const FTransform Out = Run1v1T(SrcYaw90, TgtIdentity);
+		CHECK(Out.GetLocation().Equals(FVector(-30.f, 120.f, 220.f), PosTolerance));
+	}
+}
