[PCG] Add support for generic read/write for properties

* Similar code that for attributes. Use the variant for InValues and OutValues
* Added helper to convert a FPCGMetadataAttributeDesc to PropertyBagPropertyDesc
* With unit tests for different types
* Not yet used, need to make sure we can still handle the exceptions (Color, PropertySelector)

#rnx

#Synced-CL 52681881

diff --git a/Engine/Plugins/PCG/Source/PCG/Private/Helpers/PCGPropertyHelpers.cpp b/Engine/Plugins/PCG/Source/PCG/Private/Helpers/PCGPropertyHelpers.cpp
--- a/Engine/Plugins/PCG/Source/PCG/Private/Helpers/PCGPropertyHelpers.cpp
+++ b/Engine/Plugins/PCG/Source/PCG/Private/Helpers/PCGPropertyHelpers.cpp
@@ -227,67 +227,121 @@ EPCGMetadataTypes PCGPropertyHelpers::GetMetadataTypeFromProperty(const FPropert
 	return PCGAttributeAccessorHelpers::GetMetadataTypeForProperty(InProperty);
 }
 
-FPropertyBagPropertyDesc PCGPropertyHelpers::CreatePropertyBagDescWithMetadataType(const FName InPropertyName, const EPCGMetadataTypes Type)
+FPropertyBagPropertyDesc PCGPropertyHelpers::CreatePropertyBagDescWithMetadataType(FName InPropertyName, const EPCGMetadataTypes Type)
 {
-	EPropertyBagPropertyType ValueType = EPropertyBagPropertyType::Struct;
-	const UObject* ValueTypeObject = nullptr;
+	return CreatePropertyBagDescWithMetadataType(InPropertyName, FPCGMetadataAttributeDesc{.ValueType = Type});
+}
 
-	switch (Type)
+FPropertyBagPropertyDesc PCGPropertyHelpers::CreatePropertyBagDescWithMetadataType(FName InPropertyName, const FPCGMetadataAttributeDesc& AttributeDesc)
+{
+	auto MetadataTypeToPropertyType = [](const EPCGMetadataTypes InType, const UObject* InTypeObject, EPropertyBagPropertyType& OutType, TObjectPtr<const UObject>& OutTypeObject)
 	{
-		// Simple Types
-		case EPCGMetadataTypes::Float:
-			ValueType = EPropertyBagPropertyType::Float;
-			break;
-		case EPCGMetadataTypes::Double:
-			ValueType = EPropertyBagPropertyType::Double;
-			break;
-		case EPCGMetadataTypes::Integer32:
-			ValueType = EPropertyBagPropertyType::Int32;
-			break;
-		case EPCGMetadataTypes::Integer64:
-			ValueType = EPropertyBagPropertyType::Int64;
-			break;
-		case EPCGMetadataTypes::String:
-			ValueType = EPropertyBagPropertyType::String;
-			break;
-		case EPCGMetadataTypes::Boolean:
-			ValueType = EPropertyBagPropertyType::Bool;
-			break;
-		case EPCGMetadataTypes::Name:
-			ValueType = EPropertyBagPropertyType::Name;
-			break;
+		OutType = EPropertyBagPropertyType::Struct;
+		OutTypeObject = InTypeObject;
 
-		// Value Type Objects - PropertyType is Struct by default.
-		case EPCGMetadataTypes::SoftObjectPath:
-			ValueTypeObject = TBaseStructure<FSoftObjectPath>::Get();
-			break;
-		case EPCGMetadataTypes::SoftClassPath:
-			ValueTypeObject = TBaseStructure<FSoftClassPath>::Get();
-			break;
-		case EPCGMetadataTypes::Vector:
-			ValueTypeObject = TBaseStructure<FVector>::Get();
-			break;
-		case EPCGMetadataTypes::Vector2:
-			ValueTypeObject = TBaseStructure<FVector2d>::Get();
-			break;
-		case EPCGMetadataTypes::Vector4:
-			ValueTypeObject = TBaseStructure<FVector4>::Get();
+		switch (InType)
+		{
+			// Simple Types
+			case EPCGMetadataTypes::Float:
+				OutType = EPropertyBagPropertyType::Float;
+				break;
+			case EPCGMetadataTypes::Double:
+				OutType = EPropertyBagPropertyType::Double;
+				break;
+			case EPCGMetadataTypes::Integer32:
+				OutType = EPropertyBagPropertyType::Int32;
+				break;
+			case EPCGMetadataTypes::Integer64:
+				OutType = EPropertyBagPropertyType::Int64;
+				break;
+			case EPCGMetadataTypes::String:
+				OutType = EPropertyBagPropertyType::String;
+				break;
+			case EPCGMetadataTypes::Boolean:
+				OutType = EPropertyBagPropertyType::Bool;
+				break;
+			case EPCGMetadataTypes::Name:
+				OutType = EPropertyBagPropertyType::Name;
+				break;
+			// Value Type Objects - PropertyType is Struct by default.
+			case EPCGMetadataTypes::SoftObjectPath:
+				OutTypeObject = TBaseStructure<FSoftObjectPath>::Get();
+				break;
+			case EPCGMetadataTypes::SoftClassPath:
+				OutTypeObject = TBaseStructure<FSoftClassPath>::Get();
+				break;
+			case EPCGMetadataTypes::Vector:
+				OutTypeObject = TBaseStructure<FVector>::Get();
+				break;
+			case EPCGMetadataTypes::Vector2:
+				OutTypeObject = TBaseStructure<FVector2d>::Get();
+				break;
+			case EPCGMetadataTypes::Vector4:
+				OutTypeObject = TBaseStructure<FVector4>::Get();
+				break;
+			case EPCGMetadataTypes::Transform:
+				OutTypeObject = TBaseStructure<FTransform>::Get();
+				break;
+			case EPCGMetadataTypes::Quaternion:
+				OutTypeObject = TBaseStructure<FQuat>::Get();
+				break;
+			case EPCGMetadataTypes::Rotator:
+				OutTypeObject = TBaseStructure<FRotator>::Get();
+				break;
+			case EPCGMetadataTypes::Struct:
+				// Nothing to do, already set up
+				break;
+			case EPCGMetadataTypes::Object:
+				OutType = EPropertyBagPropertyType::Object;
+				break;
+			case EPCGMetadataTypes::SoftObject:
+				OutType = EPropertyBagPropertyType::SoftObject;
+				break;
+			case EPCGMetadataTypes::Class:
+				OutType = EPropertyBagPropertyType::Class;
+				break;
+			case EPCGMetadataTypes::SoftClass:
+				OutType = EPropertyBagPropertyType::SoftClass;
+				break;
+			case EPCGMetadataTypes::Enum:
+				OutType = EPropertyBagPropertyType::Enum;
+				break;
+			default:
+				OutType = EPropertyBagPropertyType::None;
+				break;
+		}
+	};
+
+	FPropertyBagPropertyDesc Result{};
+	Result.Name = InPropertyName;
+	MetadataTypeToPropertyType(AttributeDesc.ValueType, AttributeDesc.ValueTypeObject, Result.ValueType, Result.ValueTypeObject);
+	MetadataTypeToPropertyType(AttributeDesc.KeyType, AttributeDesc.KeyTypeObject, Result.KeyType, Result.KeyTypeObject);
+
+	for (const EPCGMetadataAttributeContainerTypes ContainerType : AttributeDesc.ContainerTypes)
+	{
+		EPropertyBagContainerType PropertyBagContainerType = EPropertyBagContainerType::None;
+		switch (ContainerType)
+		{
+		case EPCGMetadataAttributeContainerTypes::Array:
+			PropertyBagContainerType = EPropertyBagContainerType::Array;
 			break;
-		case EPCGMetadataTypes::Transform:
-			ValueTypeObject = TBaseStructure<FTransform>::Get();
+		case EPCGMetadataAttributeContainerTypes::Set:
+			PropertyBagContainerType = EPropertyBagContainerType::Set;
 			break;
-		case EPCGMetadataTypes::Quaternion:
-			ValueTypeObject = TBaseStructure<FQuat>::Get();
+		case EPCGMetadataAttributeContainerTypes::Map:
+			PropertyBagContainerType = EPropertyBagContainerType::Map;
 			break;
-		case EPCGMetadataTypes::Rotator:
-			ValueTypeObject = TBaseStructure<FRotator>::Get();
+		default:
 			break;
+		}
 
-		default:
-			return {};
+		if (PropertyBagContainerType != EPropertyBagContainerType::None)
+		{
+			Result.ContainerTypes.Add(PropertyBagContainerType);
+		}
 	}
 
-	return FPropertyBagPropertyDesc(InPropertyName, ValueType, ValueTypeObject);
+	return Result;
 }
 
 UPCGParamData* PCGPropertyHelpers::ExtractPropertyAsAttributeSet(const PCGPropertyHelpers::FExtractorParameters& Parameters, FPCGContext* OptionalContext, TSet<FSoftObjectPath>* OptionalObjectTraversed, bool bQuiet)
diff --git a/Engine/Plugins/PCG/Source/PCG/Private/Metadata/Accessors/PCGPropertyAccessor.cpp b/Engine/Plugins/PCG/Source/PCG/Private/Metadata/Accessors/PCGPropertyAccessor.cpp
--- a/Engine/Plugins/PCG/Source/PCG/Private/Metadata/Accessors/PCGPropertyAccessor.cpp
+++ b/Engine/Plugins/PCG/Source/PCG/Private/Metadata/Accessors/PCGPropertyAccessor.cpp
@@ -2,6 +2,7 @@
 
 #include "Metadata/Accessors/PCGPropertyAccessor.h"
 
+#include "Metadata/PCGMetadataAttribute.h"
 #include "UObject/EnumProperty.h"
 
 #define LOCTEXT_NAMESPACE "PCGPropertyAccessor"
@@ -104,4 +105,244 @@ bool FPCGEnumPropertyAccessor::SetRangeImpl(TArrayView<const int64> InValues, in
 	return bIterateOK && bWriteOK;
 }
 
-#undef LOCTEXT_NAMESPACE
\ No newline at end of file
+FPCGPropertyGenericAccessor::FPCGPropertyGenericAccessor(const FProperty* InProperty, TArray<const FProperty*>&& ExtraProperties)
+	: Super(/*bInReadOnly=*/ false, static_cast<int16>(EPCGMetadataTypes::Unknown))
+	, IPCGPropertyChain(InProperty, std::move(ExtraProperties))
+	, Property(InProperty)
+{
+	check(Property);
+	UnderlyingDesc = FPCGMetadataAttributeDesc::CreateFromProperty(InProperty);
+	UnderlyingType = static_cast<int16>(UnderlyingDesc.ValueType);
+}
+
+bool FPCGPropertyGenericAccessor::GetRangeVirtual(PCG::Private::FOutValues OutValues, int32 Count, int32 Index, const IPCGAttributeAccessorKeys& Keys) const
+{
+	TArray<const void*, TInlineAllocator<256>> ContainerKeys;
+	PCGPropertyAccessor::GetContainerKeys(Index, Count, Keys, ContainerKeys);
+	if (ContainerKeys.IsEmpty())
+	{
+		return false;
+	}
+
+	// Update the addresses of all. We want to ignore the last property to call the getter on the container if it is a single value and it cannot be memcopied.
+	const bool bDataCanBeMemcpy = Property->HasAnyPropertyFlags(CPF_IsPlainOldData) && !Property->HasGetter();
+	const bool bShouldOperateOnContainer = OutValues.IsType<PCG::Private::FOutValuesByValue>() && !bDataCanBeMemcpy;
+	PCGPropertyAccessor::AddressOffset<const void*>(GetPropertyChain(), ContainerKeys, /*bIgnoreLastProperty=*/bShouldOperateOnContainer);
+	
+	return Visit([this, Count, &ContainerKeys, bDataCanBeMemcpy](auto&& OutValues) -> bool
+	{
+		using T = std::decay_t<decltype(OutValues)>;
+		if constexpr (std::is_same_v<T, PCG::Private::FOutValuesByValue> || std::is_same_v<T, PCG::Private::FOutValuesByPtr>)
+		{
+			if (!UnderlyingDesc.IsSingleValue())
+			{
+				return false;
+			}
+			
+			if constexpr(std::is_same_v<T, PCG::Private::FOutValuesByValue>)
+			{
+				uint8* RawOutValues = static_cast<uint8*>(OutValues.OutValues);
+				const int32 StructSize = Property->GetElementSize();
+				
+				if (bDataCanBeMemcpy)
+				{
+					for (int32 i = 0; i < Count; ++i)
+					{
+						FMemory::Memcpy(RawOutValues + i * StructSize, ContainerKeys[i], StructSize);
+					}
+				}
+				else
+				{
+					for (int32 i = 0; i < Count; ++i)
+					{
+						Property->GetValue_InContainer(ContainerKeys[i], RawOutValues + i * StructSize);
+					}
+				}
+			}
+			else // if constexpr (std::is_same_v<T, PCG::Private::FOutValuesByPtr>)
+			{
+				FMemory::Memcpy(OutValues.OutValues.GetData(), ContainerKeys.GetData(), sizeof(const void*) * Count);
+			}
+
+			return true;
+		}
+		else if constexpr (std::is_same_v<T, PCG::Private::FOutValuesAsArray>)
+		{
+			if (!UnderlyingDesc.IsArray())
+			{
+				return false;
+			}
+			
+			const FArrayProperty* ArrayProperty = CastFieldChecked<const FArrayProperty>(Property);
+			for (int32 i = 0; i < Count; ++i)
+			{
+				FScriptArrayHelper Helper(ArrayProperty, ContainerKeys[i]);
+				OutValues.OutValues[i] = {Helper.GetRawPtr(), Helper.Num()};
+			}
+			
+			return true;
+		}
+		else if constexpr (std::is_same_v<T, PCG::Private::FOutValuesAsSet>)
+		{
+			if (!UnderlyingDesc.IsSet())
+			{
+				return false;
+			}
+			
+			for (int32 i = 0; i < Count; ++i)
+			{
+				new (OutValues.OutValues[i]) FScriptSetHelper(CastFieldChecked<const FSetProperty>(Property), ContainerKeys[i]);
+			}
+			
+			return true;
+		}
+		else if constexpr (std::is_same_v<T, PCG::Private::FOutValuesAsMap>)
+		{
+			if (!UnderlyingDesc.IsMap())
+			{
+				return false;
+			}
+			
+			for (int32 i = 0; i < Count; ++i)
+			{
+				new (OutValues.OutValues[i]) FScriptMapHelper(CastFieldChecked<const FMapProperty>(Property), ContainerKeys[i]);
+			}
+			
+			return true;
+		}
+		else
+		{
+			static_assert(!std::is_same_v<T, T>, "Missing case in the visit");
+			return false;
+		}
+	}, OutValues);
+}
+
+bool FPCGPropertyGenericAccessor::SetRangeVirtual(PCG::Private::FInValues InValues, int32 Count, int32 Index, IPCGAttributeAccessorKeys& Keys, EPCGAttributeAccessorFlags Flags)
+{
+	TArray<void*, TInlineAllocator<256>> ContainerKeys;
+	PCGPropertyAccessor::GetContainerKeys(Index, Count, Keys, ContainerKeys);
+	if (ContainerKeys.IsEmpty())
+	{
+		return false;
+	}
+
+	// Update the addresses of all. We want to ignore the last property to call the setter on the container if it is a single value and it can not be memcopied.
+	const bool bDataCanBeMemcpy = Property->HasAnyPropertyFlags(CPF_IsPlainOldData) && !Property->HasSetter();
+	const bool bShouldOperateOnContainer = InValues.IsType<PCG::Private::FInValuesByValue>() && !bDataCanBeMemcpy;
+	PCGPropertyAccessor::AddressOffset<void*>(GetPropertyChain(), ContainerKeys, /*bIgnoreLastProperty=*/bShouldOperateOnContainer);
+	
+	return Visit([this, Count, &ContainerKeys, bDataCanBeMemcpy](auto&& InValues) -> bool
+	{
+		using T = std::decay_t<decltype(InValues)>;
+		if constexpr (std::is_same_v<T, PCG::Private::FInValuesByValue>)
+		{
+			if (!UnderlyingDesc.IsSingleValue())
+			{
+				return false;
+			}
+			
+			const uint8* RawInValues = static_cast<const uint8*>(InValues.InValues);
+			const int32 StructSize = Property->GetElementSize();
+			
+			if (bDataCanBeMemcpy)
+			{
+				for (int32 i = 0; i < Count; ++i)
+				{
+					FMemory::Memcpy(ContainerKeys[i], RawInValues + i * StructSize, StructSize);
+				}
+			}
+			else
+			{
+				for (int32 i = 0; i < Count; ++i)
+				{
+					Property->SetValue_InContainer(ContainerKeys[i], RawInValues + i * StructSize);
+				}
+			}
+
+			return true;
+		}
+		else if constexpr (std::is_same_v<T, PCG::Private::FInValuesAsArray>)
+		{
+			if (!UnderlyingDesc.IsArray())
+			{
+				return false;
+			}
+			
+			const FArrayProperty* ArrayProperty = CastFieldChecked<const FArrayProperty>(Property);
+			for (int32 i = 0; i < Count; ++i)
+			{
+				FScriptArrayHelper Helper(ArrayProperty, ContainerKeys[i]);
+				const int32 Num = InValues.InValues[i].template Get<1>();
+				if (ArrayProperty->Inner->HasAnyPropertyFlags(CPF_IsPlainOldData))
+				{
+					Helper.EmptyAndAddUninitializedValues(Num);
+				}
+				else
+				{
+					Helper.EmptyAndAddValues(Num);
+				}
+
+				PCG::Private::CopyArray(ArrayProperty, Helper.GetRawPtr(), InValues.InValues[i].template Get<0>(), Num);
+			}
+			
+			return true;
+		}
+		else if constexpr (std::is_same_v<T, PCG::Private::FInValuesAsSet>)
+		{
+			if (!UnderlyingDesc.IsSet())
+			{
+				return false;
+			}
+			
+			const FSetProperty* SetProperty = CastFieldChecked<const FSetProperty>(Property);
+			for (int32 i = 0; i < Count; ++i)
+			{
+				FScriptSetHelper Helper(SetProperty, ContainerKeys[i]);
+				const int32 Num = InValues.InValues[i].Num();
+				Helper.EmptyElements(Num);
+
+				for (int j = 0; j < Num; ++j)
+				{
+					Helper.AddElement(InValues.InValues[i][j]);
+				}
+			}
+			
+			return true;
+		}
+		else if constexpr (std::is_same_v<T, PCG::Private::FInValuesAsMap>)
+		{
+			if (!UnderlyingDesc.IsMap())
+			{
+				return false;
+			}
+			
+			const FMapProperty* MapProperty = CastFieldChecked<const FMapProperty>(Property);
+			for (int32 i = 0; i < Count; ++i)
+			{
+				FScriptMapHelper Helper(MapProperty, ContainerKeys[i]);
+				const int32 Num = InValues.InValues[i].Num();
+				Helper.EmptyValues(Num);
+
+				for (int j = 0; j < Num; ++j)
+				{
+					Helper.AddPair(InValues.InValues[i][j].Key, InValues.InValues[i][j].Value);
+				}
+			}
+						
+			return true;
+		}
+		else if constexpr (std::is_same_v<T, PCG::Private::FInValuesSubset>)
+		{
+			// Subset is unsupported.
+			return false;
+		}
+		else
+		{
+			static_assert(!std::is_same_v<T, T>, "Missing case in the visit");
+			return false;
+		}
+	}, InValues);
+}
+
+#undef LOCTEXT_NAMESPACE
diff --git a/Engine/Plugins/PCG/Source/PCG/Private/Metadata/PCGMetadataCommon.cpp b/Engine/Plugins/PCG/Source/PCG/Private/Metadata/PCGMetadataCommon.cpp
--- a/Engine/Plugins/PCG/Source/PCG/Private/Metadata/PCGMetadataCommon.cpp
+++ b/Engine/Plugins/PCG/Source/PCG/Private/Metadata/PCGMetadataCommon.cpp
@@ -3,6 +3,8 @@
 #include "Metadata/PCGMetadataCommon.h"
 
 #include "UObject/Class.h"
+#include "UObject/EnumProperty.h"
+#include "UObject/UnrealType.h"
 
 #include UE_INLINE_GENERATED_CPP_BY_NAME(PCGMetadataCommon)
 
@@ -23,6 +25,43 @@ bool FPCGMetadataAttributeDesc::IsSameType(const FPCGMetadataAttributeDesc& Othe
 		KeyTypeObject == Other.KeyTypeObject;
 }
 
+bool FPCGMetadataAttributeDesc::IsValid() const
+{
+	auto ValidType = [](EPCGMetadataTypes InType, const TObjectPtr<const UObject>& InTypeObject)
+	{
+		if (InType >= EPCGMetadataTypes::Count)
+		{
+			return false;
+		}
+		
+		if ((InType == EPCGMetadataTypes::Struct || InType == EPCGMetadataTypes::SoftObject || InType == EPCGMetadataTypes::Enum ||
+			InType == EPCGMetadataTypes::Object || InType == EPCGMetadataTypes::SoftClass || InType == EPCGMetadataTypes::Class) &&
+			InTypeObject == nullptr)
+		{
+			return false;
+		}
+			
+		return true;
+	};
+	
+	if (ContainerTypes.Num() >= 2)
+	{
+		return false;
+	}
+	
+	if (!ValidType(ValueType, ValueTypeObject))
+	{
+		return false;
+	}
+	
+	if (!ContainerTypes.IsEmpty() && ContainerTypes[0] == EPCGMetadataAttributeContainerTypes::Map && !ValidType(KeyType, KeyTypeObject))
+	{
+		return false;
+	}
+	
+	return true;
+}
+
 bool FPCGMetadataAttributeDesc::operator==(const FPCGMetadataAttributeDesc& Other) const
 {
 	return Name == Other.Name && IsSameType(Other);
@@ -152,6 +191,168 @@ FString FPCGMetadataAttributeDesc::GetTypeString() const
 	return Builder.ToString();
 }
 
+FPCGMetadataAttributeDesc FPCGMetadataAttributeDesc::CreateFromProperty(const FProperty* InProperty)
+{
+	if (!InProperty)
+	{
+		return {};
+	}
+
+	auto SetupProperty = [](const FProperty* InProperty, EPCGMetadataTypes& OutType, TObjectPtr<const UObject>& OutTypeObject) -> bool
+	{
+		if (!InProperty)
+		{
+			return false;
+		}
+
+		bool bValid = true;
+
+		if (const FNumericProperty* NumericProperty = CastField<FNumericProperty>(InProperty))
+		{
+			if (NumericProperty->IsFloatingPoint())
+			{
+				// As floating properties are mostly all double in UE, convert float to double attributes.
+				OutType = EPCGMetadataTypes::Double;
+			}
+			else if (NumericProperty->IsInteger())
+			{
+				// But for int32/int64 we can distinguish between the two. Everything of size 32 or less is will be an int32, 64bits integers will be int64.
+				if (NumericProperty->IsA<FInt64Property>() || NumericProperty->IsA<FUInt64Property>())
+				{
+					OutType = EPCGMetadataTypes::Integer64;
+				}
+				else
+				{
+					OutType = EPCGMetadataTypes::Integer32;
+				}
+			}
+		}
+		else if (const FBoolProperty* BoolProperty = CastField<FBoolProperty>(InProperty))
+		{
+			OutType = EPCGMetadataTypes::Boolean;
+		}
+		else if (const FStrProperty* StringProperty = CastField<FStrProperty>(InProperty))
+		{
+			OutType = EPCGMetadataTypes::String;
+		}
+		else if (const FNameProperty* NameProperty = CastField<FNameProperty>(InProperty))
+		{
+			OutType = EPCGMetadataTypes::Name;
+		}
+		else if (const FEnumProperty* EnumProperty = CastField<FEnumProperty>(InProperty))
+		{
+			OutType = EPCGMetadataTypes::Enum;
+			OutTypeObject = EnumProperty->GetEnum();
+		}
+		else if (const FSoftClassProperty* SoftClassProperty = CastField<FSoftClassProperty>(InProperty))
+		{
+			OutType = EPCGMetadataTypes::SoftClass;
+			OutTypeObject = SoftClassProperty->MetaClass;
+		}
+		else if (const FSoftObjectProperty* SoftObjectProperty = CastField<FSoftObjectProperty>(InProperty))
+		{
+			OutType = EPCGMetadataTypes::SoftObject;
+			OutTypeObject = SoftObjectProperty->PropertyClass;
+		}
+		else if (const FClassProperty* ClassProperty = CastField<FClassProperty>(InProperty))
+		{
+			OutType = EPCGMetadataTypes::Class;
+			OutTypeObject = ClassProperty->MetaClass;
+		}
+		else if (const FObjectProperty* ObjectProperty = CastField<FObjectProperty>(InProperty))
+		{
+			OutType = EPCGMetadataTypes::Object;
+			OutTypeObject = ObjectProperty->PropertyClass;
+		}
+		else if (const FStructProperty* StructProperty = CastField<FStructProperty>(InProperty))
+		{
+			if (StructProperty->Struct == TBaseStructure<FVector>::Get())
+			{
+				OutType = EPCGMetadataTypes::Vector;
+			}
+			else if (StructProperty->Struct == TBaseStructure<FVector4>::Get())
+			{
+				OutType = EPCGMetadataTypes::Vector4;
+			}
+			else if (StructProperty->Struct == TBaseStructure<FQuat>::Get())
+			{
+				OutType = EPCGMetadataTypes::Quaternion;
+			}
+			else if (StructProperty->Struct == TBaseStructure<FTransform>::Get())
+			{
+				OutType = EPCGMetadataTypes::Transform;
+			}
+			else if (StructProperty->Struct == TBaseStructure<FRotator>::Get())
+			{
+				OutType = EPCGMetadataTypes::Rotator;
+			}
+			else if (StructProperty->Struct == TBaseStructure<FVector2D>::Get())
+			{
+				OutType = EPCGMetadataTypes::Vector2;
+			}
+			else if (StructProperty->Struct == TBaseStructure<FSoftObjectPath>::Get())
+			{
+				OutType = EPCGMetadataTypes::SoftObjectPath;
+			}
+			else if (StructProperty->Struct == TBaseStructure<FSoftClassPath>::Get())
+			{
+				OutType = EPCGMetadataTypes::SoftClassPath;
+			}
+			else
+			{
+				OutType = EPCGMetadataTypes::Struct;
+				OutTypeObject = StructProperty->Struct;
+			}
+		}
+		else
+		{
+			bValid = false;
+		}
+
+		return bValid;
+	};
+
+	FPCGMetadataAttributeDesc Result;
+
+	if (const FArrayProperty* ArrayProperty = CastField<FArrayProperty>(InProperty))
+	{
+		Result.Name = ArrayProperty->GetFName();
+		Result.ContainerTypes.Add(EPCGMetadataAttributeContainerTypes::Array);
+		if (!SetupProperty(ArrayProperty->Inner, Result.ValueType, Result.ValueTypeObject))
+		{
+			return {};
+		}
+	}
+	else if (const FSetProperty* SetProperty = CastField<FSetProperty>(InProperty))
+	{
+		Result.Name = SetProperty->GetFName();
+		Result.ContainerTypes.Add(EPCGMetadataAttributeContainerTypes::Set);
+		if (!SetupProperty(SetProperty->ElementProp, Result.ValueType, Result.ValueTypeObject))
+		{
+			return {};
+		}
+	}
+	else if (const FMapProperty* MapProperty = CastField<FMapProperty>(InProperty))
+	{
+		Result.Name = MapProperty->GetFName();
+		Result.ContainerTypes.Add(EPCGMetadataAttributeContainerTypes::Map);
+		if (!SetupProperty(MapProperty->ValueProp, Result.ValueType, Result.ValueTypeObject) || !SetupProperty(MapProperty->KeyProp, Result.KeyType, Result.KeyTypeObject))
+		{
+			return {};
+		}
+	}
+	else if (InProperty)
+	{
+		Result.Name = InProperty->GetFName();
+		if (!SetupProperty(InProperty, Result.ValueType, Result.ValueTypeObject))
+		{
+			return {};
+		}
+	}
+
+	return Result;
+}
+
 int32 GetTypeHash(const FPCGMetadataAttributeDesc& Desc)
 {
 	int32 Hash = GetTypeHash(Desc.Name);
diff --git a/Engine/Plugins/PCG/Source/PCG/Public/Helpers/PCGPropertyHelpers.h b/Engine/Plugins/PCG/Source/PCG/Public/Helpers/PCGPropertyHelpers.h
--- a/Engine/Plugins/PCG/Source/PCG/Public/Helpers/PCGPropertyHelpers.h
+++ b/Engine/Plugins/PCG/Source/PCG/Public/Helpers/PCGPropertyHelpers.h
@@ -54,6 +54,14 @@ namespace PCGPropertyHelpers
 	 * @returns The property bag property type.
 	 */
 	PCG_API FPropertyBagPropertyDesc CreatePropertyBagDescWithMetadataType(FName InPropertyName, EPCGMetadataTypes Type);
+	
+	/**
+	 * Conversion between StructUtils' PropertyBag types and PCG Attribute Desc.
+	 * @param InPropertyName - The name of the property.
+	 * @param AttributeDesc - The PCG metadata attribute descriptor.
+	 * @returns The property bag property type.
+	 */
+	PCG_API FPropertyBagPropertyDesc CreatePropertyBagDescWithMetadataType(FName InPropertyName, const FPCGMetadataAttributeDesc& AttributeDesc);
 
 	struct FExtractorParameters
 	{
diff --git a/Engine/Plugins/PCG/Source/PCG/Public/Metadata/Accessors/PCGPropertyAccessor.h b/Engine/Plugins/PCG/Source/PCG/Public/Metadata/Accessors/PCGPropertyAccessor.h
--- a/Engine/Plugins/PCG/Source/PCG/Public/Metadata/Accessors/PCGPropertyAccessor.h
+++ b/Engine/Plugins/PCG/Source/PCG/Public/Metadata/Accessors/PCGPropertyAccessor.h
@@ -39,9 +39,9 @@ namespace PCGPropertyAccessor
 	}
 
 	template <typename T>
-	void AddressOffset(const TArray<const FProperty*>& InProperties, TArrayView<T> InContainerKeys)
+	void AddressOffset(const TArray<const FProperty*>& InProperties, TArrayView<T> InContainerKeys, bool bIgnoreLastProperty = false)
 	{
-		for (int32 j = 0; j < InProperties.Num(); ++j)
+		for (int32 j = 0; j < (bIgnoreLastProperty ? InProperties.Num() - 1 : InProperties.Num()); ++j)
 		{
 			const FProperty* Property = InProperties[j];
 			check(Property);
@@ -596,4 +596,23 @@ public:
 
 private:
 	const FStructProperty* Property = nullptr;
-};
\ No newline at end of file
+};
+
+/**
+* Generic accessor class for properties.
+* Do not instantiate it manually, use PCGAttributeAccessorHelpers::CreatePropertyAccessor.
+* Key supported: Generic object
+*/
+class FPCGPropertyGenericAccessor : public IPCGAttributeAccessor, IPCGPropertyChain
+{
+public:
+	using Super = IPCGAttributeAccessor;
+
+	PCG_API FPCGPropertyGenericAccessor(const FProperty* InProperty, TArray<const FProperty*>&& ExtraProperties = {});
+
+	PCG_API virtual bool GetRangeVirtual(PCG::Private::FOutValues OutValues, int32 Count, int32 Index, const IPCGAttributeAccessorKeys& Keys) const override;
+	PCG_API virtual bool SetRangeVirtual(PCG::Private::FInValues InValues, int32 Count, int32 Index, IPCGAttributeAccessorKeys& Keys, EPCGAttributeAccessorFlags Flags) override;
+
+private:
+	const FProperty* Property = nullptr;
+};
diff --git a/Engine/Plugins/PCG/Source/PCG/Public/Metadata/PCGMetadataCommon.h b/Engine/Plugins/PCG/Source/PCG/Public/Metadata/PCGMetadataCommon.h
--- a/Engine/Plugins/PCG/Source/PCG/Public/Metadata/PCGMetadataCommon.h
+++ b/Engine/Plugins/PCG/Source/PCG/Public/Metadata/PCGMetadataCommon.h
@@ -112,7 +112,6 @@ struct FPCGMetadataDomainID
 	UPROPERTY(EditAnywhere, BlueprintReadWrite, Category = "")
 	int32 CustomFlag = -1;
 
-	UPROPERTY(EditAnywhere, BlueprintReadWrite, Category = "")
 	FName DebugName = NAME_None;
 };
 
@@ -291,6 +290,8 @@ struct FPCGMetadataAttributeDesc
 	TObjectPtr<const UObject> KeyTypeObject = nullptr;
 	
 	PCG_API bool IsSameType(const FPCGMetadataAttributeDesc& Other) const;
+	
+	PCG_API bool IsValid() const;
 
 	/** Same type and same name */
 	PCG_API bool operator==(const FPCGMetadataAttributeDesc& Other) const;
@@ -313,7 +314,10 @@ struct FPCGMetadataAttributeDesc
 	PCG_API bool IsMap() const;
 
 	/** Returns the string representing the type of this descriptor. */
-	FString GetTypeString() const;
+	PCG_API FString GetTypeString() const;
+
+	/** Create a descriptor from a Property. */
+	static PCG_API FPCGMetadataAttributeDesc CreateFromProperty(const FProperty* InProperty);
 };
 
 /**
diff --git a/Engine/Plugins/PCG/Tests/Private/Metadata/Accessors/PCGPropertyAccessorTests.cpp b/Engine/Plugins/PCG/Tests/Private/Metadata/Accessors/PCGPropertyAccessorTests.cpp
new file mode 100644
--- /dev/null
+++ b/Engine/Plugins/PCG/Tests/Private/Metadata/Accessors/PCGPropertyAccessorTests.cpp
@@ -0,0 +1,295 @@
+﻿// Copyright Epic Games, Inc. All Rights Reserved.
+
+// =============================================================================
+// PCG Property Accessor Tests
+//
+// Tests for FPCGPropertyGenericAccessor Get/Set operations on Unreal properties.
+// Validates generic read/write across multiple value types, container
+// types (Array, Set, Map), property chains (struct-in-struct), non-zero property
+// offsets, and custom getter/setter support.
+//
+// Test summary:
+//  - Basic::DirectContainer:              Set/Get values through a single PropertyBag property for each type
+//  - Basic::StructInStruct:               Set/Get values through a nested struct property chain
+//  - Basic::WithOffset:                   Set/Get on the 3rd property in a bag, verifying the 1st and 2nd are untouched
+//  - GetterSetter:                        Set/Get through a UObject property with custom Setter (divides by 2) and Getter (multiplies by 2)
+//
+// Types covered by Basic:
+//  Double, FVector, FString, FPCGPoint, TArray<double>, TArray<FPCGPoint>,
+//  TSet<double>, TMap<FString, double>, EPCGMetadataTypes (enum)
+// =============================================================================
+
+#include "PCGPropertyAccessorTests.h"
+
+#include "CoreMinimal.h"
+#include "TestHarness.h"
+
+#include "PCGTestsCommon.h"
+#include "Metadata/PCGMetadataAttributeTestsCommonHelper.h"
+
+#include <catch2/generators/catch_generators.hpp>
+
+#include "Data/PCGPointArrayData.h"
+#include "Helpers/PCGPropertyHelpers.h"
+#include "Metadata/Accessors/PCGAttributeAccessorHelpers.h"
+#include "Metadata/Accessors/PCGPropertyAccessor.h"
+
+// Tests that FPCGPropertyGenericAccessor can Set and Get values for various types, through direct containers, nested struct chains, and properties at non-zero offsets within a PropertyBag.
+TEST_CASE_METHOD(PCGTests::FPCGBaseTestWithBasicContext, "PCG::Metadata::Accessors::Properties::Basic", "[PCG][Properties]")
+{
+	constexpr int32 Seed = 42;
+	constexpr int32 NumInstances = 128;
+	
+	auto Tester = [this]<typename T, typename ReadType>(PCGAttributeTestsCommonHelper::TypedAttributeTester<T, ReadType> AttributeTester)
+	{
+		static const FName PropertyName = "NewProperty";
+		
+		const FPropertyBagPropertyDesc Desc = PCGPropertyHelpers::CreatePropertyBagDescWithMetadataType(PropertyName, AttributeTester.ExpectedDesc);
+		const UPropertyBag* PropertyBag = UPropertyBag::GetOrCreateFromDescs({Desc});
+		REQUIRE_NOT_EQUAL(PropertyBag, nullptr);
+		
+		TArray<FInstancedPropertyBag> Instances;
+		TArray<void*> InstancesMemory;
+		Instances.Reserve(NumInstances);
+		InstancesMemory.Reserve(NumInstances);
+		
+		for (int32 i = 0; i < NumInstances; ++i)
+		{
+			Instances.Emplace_GetRef().InitializeFromBagStruct(PropertyBag);
+			REQUIRE(Instances[i].IsValid());
+			InstancesMemory.Emplace(Instances[i].GetMutableValue().GetMemory());
+		}
+		
+		const FPropertyBagPropertyDesc* AddedDesc = PropertyBag->FindPropertyDescByName(PropertyName);
+		REQUIRE_NOT_EQUAL(AddedDesc, nullptr);
+		
+		SECTION("Direct Container")
+		{
+			FRandomStream RandomStream(Seed);
+			TArray<T> ValuesToAdd{};
+			ValuesToAdd.Reserve(NumInstances);
+			
+			for (int32 i = 0; i < NumInstances; ++i)
+			{
+				AttributeTester.GenerateRandom(RandomStream, ValuesToAdd.Emplace_GetRef());
+			}
+			
+			FPCGPropertyGenericAccessor Accessor(AddedDesc->CachedProperty);
+			FPCGAttributeAccessorKeysGenericPtrs Keys(InstancesMemory);
+		
+			REQUIRE(Accessor.GetUnderlyingDesc().IsSameType(AttributeTester.ExpectedDesc));
+		
+			REQUIRE(Accessor.SetRange<T>(ValuesToAdd, /*Index=*/0, Keys));
+		
+			TArray<ReadType> ReadValues;
+			ReadValues.SetNum(NumInstances);
+			REQUIRE(Accessor.GetRange<ReadType>(ReadValues,  /*Index=*/0, Keys));
+		
+			CHECK_THAT(ValuesToAdd, Catch::Matchers::RangeEquals(ReadValues, [&AttributeTester](const auto& a, const auto& b){ return AttributeTester.Verify(a, b); } ));
+		}
+		
+		SECTION("Struct in Struct")
+		{
+			static const FName InnerStructName = "Inner";
+			TArray<FInstancedPropertyBag> OuterInstances;
+			TArray<void*> OuterInstancesMemory;
+			OuterInstances.Reserve(NumInstances);
+			OuterInstancesMemory.Reserve(NumInstances);
+			
+			FPropertyBagPropertyDesc InnerDesc(InnerStructName, EPropertyBagPropertyType::Struct, PropertyBag);
+			const UPropertyBag* OuterPropertyBag = UPropertyBag::GetOrCreateFromDescs({InnerDesc});
+			
+			for (int32 i = 0; i < NumInstances; ++i)
+			{
+				OuterInstances.Emplace_GetRef().InitializeFromBagStruct(OuterPropertyBag);
+				REQUIRE(OuterInstances[i].IsValid());
+				OuterInstancesMemory.Emplace(OuterInstances[i].GetMutableValue().GetMemory());
+			}
+			
+			const FPropertyBagPropertyDesc* InnerAddedDesc = OuterPropertyBag->FindPropertyDescByName(InnerStructName);
+			REQUIRE_NOT_EQUAL(InnerAddedDesc, nullptr);
+			
+			// Offset seed to get another value
+			FRandomStream RandomStream(Seed+1);
+			TArray<T> ValuesToAdd{};
+			ValuesToAdd.Reserve(NumInstances);
+			
+			for (int32 i = 0; i < NumInstances; ++i)
+			{
+				AttributeTester.GenerateRandom(RandomStream, ValuesToAdd.Emplace_GetRef());
+			}
+			
+			FPCGPropertyGenericAccessor Accessor(AddedDesc->CachedProperty, {InnerAddedDesc->CachedProperty});
+			FPCGAttributeAccessorKeysGenericPtrs Keys(OuterInstancesMemory);
+		
+			REQUIRE(Accessor.GetUnderlyingDesc().IsSameType(AttributeTester.ExpectedDesc));
+		
+			REQUIRE(Accessor.SetRange<T>(ValuesToAdd, /*Index=*/0, Keys));
+		
+			TArray<ReadType> ReadValues;
+			ReadValues.SetNum(NumInstances);
+			REQUIRE(Accessor.GetRange<ReadType>(ReadValues,  /*Index=*/0, Keys));
+		
+			CHECK_THAT(ValuesToAdd, Catch::Matchers::RangeEquals(ReadValues, [&AttributeTester](const auto& a, const auto& b){ return AttributeTester.Verify(a, b); } ));
+		}
+		
+		SECTION("With offset")
+		{
+			// Add the same property thrice, and try to access the last one
+			const FPropertyBagPropertyDesc Desc1 = PCGPropertyHelpers::CreatePropertyBagDescWithMetadataType(TEXT("FirstProp"), AttributeTester.ExpectedDesc);
+			const FPropertyBagPropertyDesc Desc2 = PCGPropertyHelpers::CreatePropertyBagDescWithMetadataType(TEXT("SecondProp"), AttributeTester.ExpectedDesc);
+			const FPropertyBagPropertyDesc Desc3 = PCGPropertyHelpers::CreatePropertyBagDescWithMetadataType(PropertyName, AttributeTester.ExpectedDesc);
+			const UPropertyBag* BulkyPropertyBag = UPropertyBag::GetOrCreateFromDescs({Desc1, Desc2, Desc3});
+			REQUIRE_NOT_EQUAL(BulkyPropertyBag, nullptr);
+		
+			TArray<FInstancedPropertyBag> BulkyInstances;
+			TArray<void*> BulkyInstancesMemory;
+			BulkyInstances.Reserve(NumInstances);
+			BulkyInstancesMemory.Reserve(NumInstances);
+		
+			for (int32 i = 0; i < NumInstances; ++i)
+			{
+				BulkyInstances.Emplace_GetRef().InitializeFromBagStruct(BulkyPropertyBag);
+				REQUIRE(BulkyInstances[i].IsValid());
+				BulkyInstancesMemory.Emplace(BulkyInstances[i].GetMutableValue().GetMemory());
+			}
+		
+			const FPropertyBagPropertyDesc* BulkyAddedDesc1 = BulkyPropertyBag->FindPropertyDescByName(TEXT("FirstProp"));
+			const FPropertyBagPropertyDesc* BulkyAddedDesc2 = BulkyPropertyBag->FindPropertyDescByName(TEXT("SecondProp"));
+			const FPropertyBagPropertyDesc* BulkyAddedDesc3 = BulkyPropertyBag->FindPropertyDescByName(PropertyName);
+			REQUIRE_NOT_EQUAL(BulkyAddedDesc1, nullptr);
+			REQUIRE_NOT_EQUAL(BulkyAddedDesc2, nullptr);
+			REQUIRE_NOT_EQUAL(BulkyAddedDesc3, nullptr);
+			
+			// Offset seed to get another value
+			FRandomStream RandomStream(Seed+2);
+			TArray<T> ValuesToAdd{};
+			ValuesToAdd.Reserve(NumInstances);
+		
+			for (int32 i = 0; i < NumInstances; ++i)
+			{
+				AttributeTester.GenerateRandom(RandomStream, ValuesToAdd.Emplace_GetRef());
+			}
+		
+			FPCGPropertyGenericAccessor Accessor(BulkyAddedDesc3->CachedProperty);
+			FPCGAttributeAccessorKeysGenericPtrs Keys(BulkyInstancesMemory);
+
+			REQUIRE(Accessor.GetUnderlyingDesc().IsSameType(AttributeTester.ExpectedDesc));
+
+			REQUIRE(Accessor.SetRange<T>(ValuesToAdd, /*Index=*/0, Keys));
+			
+			// Verify that the values are set correctly
+			FInstancedPropertyBag EmptyInstance;
+			EmptyInstance.InitializeFromBagStruct(BulkyPropertyBag);
+			for (int32 i = 0; i < NumInstances; ++i)
+			{
+				REQUIRE(BulkyAddedDesc1->CachedProperty->Identical_InContainer(BulkyInstances[i].GetValue().GetMemory(), EmptyInstance.GetValue().GetMemory()));
+				REQUIRE(BulkyAddedDesc2->CachedProperty->Identical_InContainer(BulkyInstances[i].GetValue().GetMemory(), EmptyInstance.GetValue().GetMemory()));
+				REQUIRE_FALSE(BulkyAddedDesc3->CachedProperty->Identical_InContainer(BulkyInstances[i].GetValue().GetMemory(), EmptyInstance.GetValue().GetMemory()));
+			}
+	
+			TArray<ReadType> ReadValues;
+			ReadValues.SetNum(NumInstances);
+			REQUIRE(Accessor.GetRange<ReadType>(ReadValues,  /*Index=*/0, Keys));
+	
+			CHECK_THAT(ValuesToAdd, Catch::Matchers::RangeEquals(ReadValues, [&AttributeTester](const auto& a, const auto& b){ return AttributeTester.Verify(a, b); } ));
+		}
+	};
+	
+	SECTION("Double Attribute")
+	{
+		Tester(PCGAttributeTestsCommonHelper::DoubleTester());
+	}
+	
+	SECTION("Vector Attribute")
+	{
+		Tester(PCGAttributeTestsCommonHelper::VectorTester_Basic());
+	}
+	
+	SECTION("String Attribute")
+	{
+		Tester(PCGAttributeTestsCommonHelper::StringTester());
+	}
+	
+	SECTION("FPCGPoint Attribute")
+	{		
+		Tester(PCGAttributeTestsCommonHelper::FPCGPointTester());
+	}
+	
+	SECTION("Array of Double Attribute")
+	{		
+		Tester(PCGAttributeTestsCommonHelper::ArrayDoubleTester());
+	}
+	
+	SECTION("Array of FPCGPoint Attribute")
+	{	
+		Tester(PCGAttributeTestsCommonHelper::ArrayFPCGPointTester());
+	}
+	
+	SECTION("Set of Double Attribute")
+	{		
+		Tester(PCGAttributeTestsCommonHelper::SetDoubleTester());
+	}
+	
+	SECTION("Map of String -> Double Attribute")
+	{
+		Tester(PCGAttributeTestsCommonHelper::MapStringDoubleTester());
+	}
+	
+	SECTION("Enum attribute")
+	{
+		Tester(PCGAttributeTestsCommonHelper::EnumTester());
+	}
+}
+
+// Tests that FPCGPropertyGenericAccessor correctly invokes custom UObject Setter/Getter when writing and reading a property, verifying the stored value is transformed (divided by 2) and the read value round-trips back to the original.
+TEST_CASE_METHOD(PCGTests::FPCGBaseTestWithBasicContext, "PCG::Metadata::Accessors::Properties::GetterSetter", "[PCG][Properties]")
+{
+	constexpr int32 Seed = 42;
+	constexpr int32 NumInstances = 128;
+	
+	TArray<UPCGPropertyAccessorTestsGetterSetter*> Instances;
+	TArray<void*> InstancesMemory;
+	Instances.Reserve(NumInstances);
+	InstancesMemory.Reserve(NumInstances);
+			
+	for (int32 i = 0; i < NumInstances; ++i)
+	{
+		Instances.Emplace(FPCGContext::NewObject_AnyThread<UPCGPropertyAccessorTestsGetterSetter>(GetContext()));
+		InstancesMemory.Emplace(Instances[i]);
+	}
+			
+	// Offset seed to get another value
+	FRandomStream RandomStream(Seed+1);
+	TArray<double> ValuesToAdd{};
+	ValuesToAdd.Reserve(NumInstances);
+			
+	for (int32 i = 0; i < NumInstances; ++i)
+	{
+		ValuesToAdd.Emplace(RandomStream.FRand());
+	}
+	
+	const FProperty* Property = UPCGPropertyAccessorTestsGetterSetter::StaticClass()->FindPropertyByName(GET_MEMBER_NAME_CHECKED(UPCGPropertyAccessorTestsGetterSetter, Value));
+	REQUIRE(Property);
+	REQUIRE(Property->HasSetter());
+	REQUIRE(Property->HasGetter());
+	
+	FPCGPropertyGenericAccessor Accessor(Property);
+	FPCGAttributeAccessorKeysGenericPtrs Keys(InstancesMemory);
+		
+	REQUIRE(Accessor.GetUnderlyingDesc().IsSameType(PCG::Private::GetDefaultAttributeDesc<double>()));
+	
+	REQUIRE(Accessor.SetRange<double>(ValuesToAdd, /*Index=*/0, Keys));
+	
+	// Verify that the value is stored as expected
+	TArray<double> StoredValues;
+	StoredValues.Reserve(NumInstances);
+	Algo::Transform(Instances, StoredValues, [](const UPCGPropertyAccessorTestsGetterSetter* Obj) { return Obj->Value; });
+	REQUIRE_THAT(ValuesToAdd, Catch::Matchers::RangeEquals(StoredValues, [](const double a, const double b){ return FMath::IsNearlyEqual(a / 2.0, b); }));
+		
+	TArray<double> ReadValues;
+	ReadValues.SetNum(NumInstances);
+	REQUIRE(Accessor.GetRange<double>(ReadValues,  /*Index=*/0, Keys));
+		
+	CHECK_THAT(ValuesToAdd, Catch::Matchers::RangeEquals(ReadValues, [](const double a, const double b){ return FMath::IsNearlyEqual(a, b); }));
+}
\ No newline at end of file
diff --git a/Engine/Plugins/PCG/Tests/Private/Metadata/Accessors/PCGPropertyAccessorTests.h b/Engine/Plugins/PCG/Tests/Private/Metadata/Accessors/PCGPropertyAccessorTests.h
new file mode 100644
--- /dev/null
+++ b/Engine/Plugins/PCG/Tests/Private/Metadata/Accessors/PCGPropertyAccessorTests.h
@@ -0,0 +1,30 @@
+﻿// Copyright Epic Games, Inc. All Rights Reserved.
+
+#pragma once
+
+#include "CoreMinimal.h"
+#include "UObject/Object.h"
+
+#include "PCGPropertyAccessorTests.generated.h"
+
+UCLASS(BlueprintInternalUseOnly)
+class UPCGPropertyAccessorTestsGetterSetter : public UObject
+{
+	GENERATED_BODY()
+
+public:
+	UFUNCTION(BlueprintCallable, BlueprintInternalUseOnly, Category="Test")
+	void SetValue(double InValue)
+	{
+		Value = InValue / 2.0;
+	}
+
+	UFUNCTION(BlueprintPure, BlueprintInternalUseOnly, Category="Test")
+	double GetValue() const
+	{
+		return Value * 2.0;
+	}
+
+	UPROPERTY(BlueprintInternalUseOnly, BlueprintReadWrite, EditAnywhere, Category="Test", Setter=SetValue, Getter=GetValue)
+	double Value = 0.0;
+};
\ No newline at end of file
diff --git a/Engine/Plugins/PCG/Tests/Private/Metadata/PCGMetadataAttributeTestsCommonHelper.h b/Engine/Plugins/PCG/Tests/Private/Metadata/PCGMetadataAttributeTestsCommonHelper.h
--- a/Engine/Plugins/PCG/Tests/Private/Metadata/PCGMetadataAttributeTestsCommonHelper.h
+++ b/Engine/Plugins/PCG/Tests/Private/Metadata/PCGMetadataAttributeTestsCommonHelper.h
@@ -332,7 +332,8 @@ namespace PCGAttributeTestsCommonHelper
 	{
 		auto GenerateRandom = [](FRandomStream& RandomStream, EPCGMetadataTypes& OutValue)
 			{
-				OutValue = static_cast<EPCGMetadataTypes>(RandomStream.RandRange(0, static_cast<int32>(EPCGMetadataTypes::Count)));
+				// Make sure to not generate a 0, to differentiate from an unset value.
+				OutValue = static_cast<EPCGMetadataTypes>(RandomStream.RandRange(1, static_cast<int32>(EPCGMetadataTypes::Count)));
 			};
 		
 		auto Verify = [](const EPCGMetadataTypes& LHS, const EPCGMetadataTypes& RHS) { return LHS == RHS; };
