[PCG] Fix unit test throwing errors for Property Bags

#rnx
#Synced-CL 52681881

diff --git a/Engine/Plugins/PCG/Tests/Private/Metadata/Accessors/PCGPropertyAccessorTests.cpp b/Engine/Plugins/PCG/Tests/Private/Metadata/Accessors/PCGPropertyAccessorTests.cpp
--- a/Engine/Plugins/PCG/Tests/Private/Metadata/Accessors/PCGPropertyAccessorTests.cpp
+++ b/Engine/Plugins/PCG/Tests/Private/Metadata/Accessors/PCGPropertyAccessorTests.cpp
@@ -43,28 +43,33 @@ TEST_CASE_METHOD(PCGTests::FPCGBaseTestWithBasicContext, "PCG::Metadata::Accesso
 	auto Tester = [this]<typename T, typename ReadType>(PCGAttributeTestsCommonHelper::TypedAttributeTester<T, ReadType> AttributeTester)
 	{
 		static const FName PropertyName = "NewProperty";
-		
+
 		const FPropertyBagPropertyDesc Desc = PCGPropertyHelpers::CreatePropertyBagDescWithMetadataType(PropertyName, AttributeTester.ExpectedDesc);
-		const UPropertyBag* PropertyBag = UPropertyBag::GetOrCreateFromDescs({Desc});
-		REQUIRE_NOT_EQUAL(PropertyBag, nullptr);
 		
-		TArray<FInstancedPropertyBag> Instances;
-		TArray<void*> InstancesMemory;
-		Instances.Reserve(NumInstances);
-		InstancesMemory.Reserve(NumInstances);
+		// @hack: Add a dummy "complex" property to bypass the ZeroConstructor flag set on the property bag if all the properties are "simple".
+		// At the moment, having a property bag into a property bag doesn't increment the property bag ref count correctly if the property bag is marked "ZeroConstructor".
+		const FPropertyBagPropertyDesc ComplexDesc("Dummy", EPropertyBagContainerType::None, EPropertyBagPropertyType::Text);
 		
-		for (int32 i = 0; i < NumInstances; ++i)
-		{
-			Instances.Emplace_GetRef().InitializeFromBagStruct(PropertyBag);
-			REQUIRE(Instances[i].IsValid());
-			InstancesMemory.Emplace(Instances[i].GetMutableValue().GetMemory());
-		}
+		const UPropertyBag* PropertyBag = UPropertyBag::GetOrCreateFromDescs({Desc, ComplexDesc});
+		REQUIRE_NOT_EQUAL(PropertyBag, nullptr);
 		
 		const FPropertyBagPropertyDesc* AddedDesc = PropertyBag->FindPropertyDescByName(PropertyName);
 		REQUIRE_NOT_EQUAL(AddedDesc, nullptr);
 		
 		SECTION("Direct Container")
 		{
+			TArray<FInstancedPropertyBag> Instances;
+			TArray<void*> InstancesMemory;
+			Instances.Reserve(NumInstances);
+			InstancesMemory.Reserve(NumInstances);
+			
+			for (int32 i = 0; i < NumInstances; ++i)
+			{
+				Instances.Emplace_GetRef().InitializeFromBagStruct(PropertyBag);
+				REQUIRE(Instances[i].IsValid());
+				InstancesMemory.Add(Instances[i].GetMutableValue().GetMemory());
+			}
+			
 			FRandomStream RandomStream(Seed);
 			TArray<T> ValuesToAdd{};
 			ValuesToAdd.Reserve(NumInstances);
@@ -103,7 +108,7 @@ TEST_CASE_METHOD(PCGTests::FPCGBaseTestWithBasicContext, "PCG::Metadata::Accesso
 			{
 				OuterInstances.Emplace_GetRef().InitializeFromBagStruct(OuterPropertyBag);
 				REQUIRE(OuterInstances[i].IsValid());
-				OuterInstancesMemory.Emplace(OuterInstances[i].GetMutableValue().GetMemory());
+				OuterInstancesMemory.Add(OuterInstances[i].GetMutableValue().GetMemory());
 			}
 			
 			const FPropertyBagPropertyDesc* InnerAddedDesc = OuterPropertyBag->FindPropertyDescByName(InnerStructName);
@@ -125,75 +130,77 @@ TEST_CASE_METHOD(PCGTests::FPCGBaseTestWithBasicContext, "PCG::Metadata::Accesso
 			REQUIRE(Accessor.GetUnderlyingDesc().IsSameType(AttributeTester.ExpectedDesc));
 		
 			REQUIRE(Accessor.SetRange<T>(ValuesToAdd, /*Index=*/0, Keys));
-		
+			
 			TArray<ReadType> ReadValues;
 			ReadValues.SetNum(NumInstances);
 			REQUIRE(Accessor.GetRange<ReadType>(ReadValues,  /*Index=*/0, Keys));
-		
+			
 			CHECK_THAT(ValuesToAdd, Catch::Matchers::RangeEquals(ReadValues, [&AttributeTester](const auto& a, const auto& b){ return AttributeTester.Verify(a, b); } ));
 		}
+
+		 SECTION("With offset")
+		 {
+		 	// Add the same property thrice, and try to access the last one
+			// @hack: Add a dummy "complex" property to bypass the ZeroConstructor flag set on the property bag if all the properties are "simple".
+			// At the moment, having a property bag into a property bag doesn't increment the property bag ref count correctly if the property bag is marked "ZeroConstructor".
+			const FPropertyBagPropertyDesc Desc1("Dummy", EPropertyBagContainerType::None, EPropertyBagPropertyType::Text);
+		 	const FPropertyBagPropertyDesc Desc2 = PCGPropertyHelpers::CreatePropertyBagDescWithMetadataType("SecondProp", AttributeTester.ExpectedDesc);
+		 	const FPropertyBagPropertyDesc Desc3 = PCGPropertyHelpers::CreatePropertyBagDescWithMetadataType(PropertyName, AttributeTester.ExpectedDesc);
+		 	const UPropertyBag* BulkyPropertyBag = UPropertyBag::GetOrCreateFromDescs({Desc1, Desc2, Desc3});
+		 	REQUIRE_NOT_EQUAL(BulkyPropertyBag, nullptr);
 		
-		SECTION("With offset")
-		{
-			// Add the same property thrice, and try to access the last one
-			const FPropertyBagPropertyDesc Desc1 = PCGPropertyHelpers::CreatePropertyBagDescWithMetadataType(TEXT("FirstProp"), AttributeTester.ExpectedDesc);
-			const FPropertyBagPropertyDesc Desc2 = PCGPropertyHelpers::CreatePropertyBagDescWithMetadataType(TEXT("SecondProp"), AttributeTester.ExpectedDesc);
-			const FPropertyBagPropertyDesc Desc3 = PCGPropertyHelpers::CreatePropertyBagDescWithMetadataType(PropertyName, AttributeTester.ExpectedDesc);
-			const UPropertyBag* BulkyPropertyBag = UPropertyBag::GetOrCreateFromDescs({Desc1, Desc2, Desc3});
-			REQUIRE_NOT_EQUAL(BulkyPropertyBag, nullptr);
+		 	TArray<FInstancedPropertyBag> BulkyInstances;
+		 	TArray<void*> BulkyInstancesMemory;
+		 	BulkyInstances.Reserve(NumInstances);
+		 	BulkyInstancesMemory.Reserve(NumInstances);
 		
-			TArray<FInstancedPropertyBag> BulkyInstances;
-			TArray<void*> BulkyInstancesMemory;
-			BulkyInstances.Reserve(NumInstances);
-			BulkyInstancesMemory.Reserve(NumInstances);
+		 	for (int32 i = 0; i < NumInstances; ++i)
+		 	{
+		 		BulkyInstances.Emplace_GetRef().InitializeFromBagStruct(BulkyPropertyBag);
+		 		REQUIRE(BulkyInstances[i].IsValid());
+		 		BulkyInstancesMemory.Add(BulkyInstances[i].GetMutableValue().GetMemory());
+		 	}
 		
-			for (int32 i = 0; i < NumInstances; ++i)
-			{
-				BulkyInstances.Emplace_GetRef().InitializeFromBagStruct(BulkyPropertyBag);
-				REQUIRE(BulkyInstances[i].IsValid());
-				BulkyInstancesMemory.Emplace(BulkyInstances[i].GetMutableValue().GetMemory());
-			}
+		 	const FPropertyBagPropertyDesc* BulkyAddedDesc1 = BulkyPropertyBag->FindPropertyDescByName("Dummy");
+		 	const FPropertyBagPropertyDesc* BulkyAddedDesc2 = BulkyPropertyBag->FindPropertyDescByName("SecondProp");
+		 	const FPropertyBagPropertyDesc* BulkyAddedDesc3 = BulkyPropertyBag->FindPropertyDescByName(PropertyName);
+		 	REQUIRE_NOT_EQUAL(BulkyAddedDesc1, nullptr);
+		 	REQUIRE_NOT_EQUAL(BulkyAddedDesc2, nullptr);
+		 	REQUIRE_NOT_EQUAL(BulkyAddedDesc3, nullptr);
+		 	
+		 	// Offset seed to get another value
+		 	FRandomStream RandomStream(Seed+2);
+		 	TArray<T> ValuesToAdd{};
+		 	ValuesToAdd.Reserve(NumInstances);
 		
-			const FPropertyBagPropertyDesc* BulkyAddedDesc1 = BulkyPropertyBag->FindPropertyDescByName(TEXT("FirstProp"));
-			const FPropertyBagPropertyDesc* BulkyAddedDesc2 = BulkyPropertyBag->FindPropertyDescByName(TEXT("SecondProp"));
-			const FPropertyBagPropertyDesc* BulkyAddedDesc3 = BulkyPropertyBag->FindPropertyDescByName(PropertyName);
-			REQUIRE_NOT_EQUAL(BulkyAddedDesc1, nullptr);
-			REQUIRE_NOT_EQUAL(BulkyAddedDesc2, nullptr);
-			REQUIRE_NOT_EQUAL(BulkyAddedDesc3, nullptr);
-			
-			// Offset seed to get another value
-			FRandomStream RandomStream(Seed+2);
-			TArray<T> ValuesToAdd{};
-			ValuesToAdd.Reserve(NumInstances);
+		 	for (int32 i = 0; i < NumInstances; ++i)
+		 	{
+		 		AttributeTester.GenerateRandom(RandomStream, ValuesToAdd.Emplace_GetRef());
+		 	}
 		
-			for (int32 i = 0; i < NumInstances; ++i)
-			{
-				AttributeTester.GenerateRandom(RandomStream, ValuesToAdd.Emplace_GetRef());
-			}
+		 	FPCGPropertyGenericAccessor Accessor(BulkyAddedDesc3->CachedProperty);
+		 	FPCGAttributeAccessorKeysGenericPtrs Keys(BulkyInstancesMemory);
 		
-			FPCGPropertyGenericAccessor Accessor(BulkyAddedDesc3->CachedProperty);
-			FPCGAttributeAccessorKeysGenericPtrs Keys(BulkyInstancesMemory);
-
-			REQUIRE(Accessor.GetUnderlyingDesc().IsSameType(AttributeTester.ExpectedDesc));
-
-			REQUIRE(Accessor.SetRange<T>(ValuesToAdd, /*Index=*/0, Keys));
-			
-			// Verify that the values are set correctly
-			FInstancedPropertyBag EmptyInstance;
-			EmptyInstance.InitializeFromBagStruct(BulkyPropertyBag);
-			for (int32 i = 0; i < NumInstances; ++i)
-			{
-				REQUIRE(BulkyAddedDesc1->CachedProperty->Identical_InContainer(BulkyInstances[i].GetValue().GetMemory(), EmptyInstance.GetValue().GetMemory()));
-				REQUIRE(BulkyAddedDesc2->CachedProperty->Identical_InContainer(BulkyInstances[i].GetValue().GetMemory(), EmptyInstance.GetValue().GetMemory()));
-				REQUIRE_FALSE(BulkyAddedDesc3->CachedProperty->Identical_InContainer(BulkyInstances[i].GetValue().GetMemory(), EmptyInstance.GetValue().GetMemory()));
-			}
-	
-			TArray<ReadType> ReadValues;
-			ReadValues.SetNum(NumInstances);
-			REQUIRE(Accessor.GetRange<ReadType>(ReadValues,  /*Index=*/0, Keys));
-	
-			CHECK_THAT(ValuesToAdd, Catch::Matchers::RangeEquals(ReadValues, [&AttributeTester](const auto& a, const auto& b){ return AttributeTester.Verify(a, b); } ));
-		}
+		 	REQUIRE(Accessor.GetUnderlyingDesc().IsSameType(AttributeTester.ExpectedDesc));
+		
+		 	REQUIRE(Accessor.SetRange<T>(ValuesToAdd, /*Index=*/0, Keys));
+		 	
+		 	// Verify that the values are set correctly
+		 	FInstancedPropertyBag EmptyInstance;
+		 	EmptyInstance.InitializeFromBagStruct(BulkyPropertyBag);
+		 	for (int32 i = 0; i < NumInstances; ++i)
+		 	{
+		 		REQUIRE(BulkyAddedDesc1->CachedProperty->Identical_InContainer(BulkyInstances[i].GetValue().GetMemory(), EmptyInstance.GetValue().GetMemory()));
+		 		REQUIRE(BulkyAddedDesc2->CachedProperty->Identical_InContainer(BulkyInstances[i].GetValue().GetMemory(), EmptyInstance.GetValue().GetMemory()));
+		 		REQUIRE_FALSE(BulkyAddedDesc3->CachedProperty->Identical_InContainer(BulkyInstances[i].GetValue().GetMemory(), EmptyInstance.GetValue().GetMemory()));
+		 	}
+		
+		 	TArray<ReadType> ReadValues;
+		 	ReadValues.SetNum(NumInstances);
+		 	REQUIRE(Accessor.GetRange<ReadType>(ReadValues,  /*Index=*/0, Keys));
+		
+		 	CHECK_THAT(ValuesToAdd, Catch::Matchers::RangeEquals(ReadValues, [&AttributeTester](const auto& a, const auto& b){ return AttributeTester.Verify(a, b); } ));
+		 }
 	};
 	
 	SECTION("Double Attribute")
@@ -255,8 +262,8 @@ TEST_CASE_METHOD(PCGTests::FPCGBaseTestWithBasicContext, "PCG::Metadata::Accesso
 			
 	for (int32 i = 0; i < NumInstances; ++i)
 	{
-		Instances.Emplace(FPCGContext::NewObject_AnyThread<UPCGPropertyAccessorTestsGetterSetter>(GetContext()));
-		InstancesMemory.Emplace(Instances[i]);
+		Instances.Add(FPCGContext::NewObject_AnyThread<UPCGPropertyAccessorTestsGetterSetter>(GetContext()));
+		InstancesMemory.Add(Instances[i]);
 	}
 			
 	// Offset seed to get another value
@@ -266,7 +273,7 @@ TEST_CASE_METHOD(PCGTests::FPCGBaseTestWithBasicContext, "PCG::Metadata::Accesso
 			
 	for (int32 i = 0; i < NumInstances; ++i)
 	{
-		ValuesToAdd.Emplace(RandomStream.FRand());
+		ValuesToAdd.Add(RandomStream.FRand());
 	}
 	
 	const FProperty* Property = UPCGPropertyAccessorTestsGetterSetter::StaticClass()->FindPropertyByName(GET_MEMBER_NAME_CHECKED(UPCGPropertyAccessorTestsGetterSetter, Value));
