[PCG] Fix order of priority for broadcast constructible

* This was broken with the refactor, broadcast was the priority but now it was constructible
* Revert behavior to previous
* Add tests for problematic types (aka Vector4)

#rnx
#Synced-CL 52758271

diff --git a/Engine/Plugins/PCG/Source/PCG/Public/Metadata/Accessors/IPCGAttributeAccessor.h b/Engine/Plugins/PCG/Source/PCG/Public/Metadata/Accessors/IPCGAttributeAccessor.h
--- a/Engine/Plugins/PCG/Source/PCG/Public/Metadata/Accessors/IPCGAttributeAccessor.h
+++ b/Engine/Plugins/PCG/Source/PCG/Public/Metadata/Accessors/IPCGAttributeAccessor.h
@@ -50,18 +50,19 @@ namespace PCG::Private
 	
 	inline EPCGAttributeAccessorOperation GetOpType(const FPCGMetadataAttributeDesc& ReadDesc, const FPCGMetadataAttributeDesc& WriteDesc, EPCGAttributeAccessorFlags Flags)
 	{
+		// Priority list if SameType -> Broadcast -> Constructible
 		if (ReadDesc.IsSameType(WriteDesc))
 		{
 			return EPCGAttributeAccessorOperation::SameType;
 		}
-		else if (!!(Flags & EPCGAttributeAccessorFlags::AllowConstructible) && PCG::Private::IsConstructible(WriteDesc, ReadDesc))
-		{
-			return EPCGAttributeAccessorOperation::Construct;
-		}
 		else if (!!(Flags & EPCGAttributeAccessorFlags::AllowBroadcast) && PCG::Private::IsBroadcastable(ReadDesc, WriteDesc))
 		{
 			return EPCGAttributeAccessorOperation::Broadcast;
 		}
+		else if (!!(Flags & EPCGAttributeAccessorFlags::AllowConstructible) && PCG::Private::IsConstructible(WriteDesc, ReadDesc))
+		{
+			return EPCGAttributeAccessorOperation::Construct;
+		}
 		
 		if (IsEquivalentDesc<FString>(WriteDesc))
 		{
diff --git a/Engine/Plugins/PCG/Tests/Private/Metadata/Accessors/PCGAttributeAccessorTests.cpp b/Engine/Plugins/PCG/Tests/Private/Metadata/Accessors/PCGAttributeAccessorTests.cpp
--- a/Engine/Plugins/PCG/Tests/Private/Metadata/Accessors/PCGAttributeAccessorTests.cpp
+++ b/Engine/Plugins/PCG/Tests/Private/Metadata/Accessors/PCGAttributeAccessorTests.cpp
@@ -24,6 +24,7 @@
 //  - SetRangePoints:                    SetRange using point-based keys writes through entry indirection
 //  - GetWithConstructibleType:          Get/Set with constructible type conversion (float<->double) using AllowConstructible flag
 //  - GetWithBroadcastableType:          Get/Set with broadcastable type conversion (float->FVector/FString) using AllowBroadcast flag
+//  - GetWithBroadcastAndConstructibleType: Get/Set with double->FVector4 conversion using AllowBroadcastAndConstructible flag
 //  - SoftObjectPathToStringWithStrictType: FSoftObjectPath/FSoftClassPath to FString conversion with StrictType flag
 //  - CopyToSameType:                     CopyTo copies values between same-type accessors (float, vector, string)
 //  - CopyToPartialRange:                 CopyTo copies a sub-range leaving untouched entries at their defaults
@@ -839,6 +840,61 @@ TEST_CASE_METHOD(PCGTests::FPCGBaseTestWithBasicContext, "PCG::Accessor::Attribu
 	}
 }
 
+/**
+ * Tests Get/Set with broadcastable and constructible type conversion (double->FVector4). Verifies that double->FVector4 conversion succeeds with AllowBroadcastAndConstructible, for both single and range operations and it
+ * is the broadcast operation that works, to be sure we are following the previous behavior before the refactor of the accessors.
+ */
+TEST_CASE_METHOD(PCGTests::FPCGBaseTestWithBasicContext, "PCG::Accessor::Attribute::GetWithBroadcastAndConstructibleType", "[PCG][Accessor][Attribute]")
+{
+	UPCGBasePointData* PointData = CreatePointData();
+
+	constexpr double DoubleValue = 2.71828;
+	FPCGMetadataAttribute<double>* DoubleAttribute = PointData->Metadata->CreateAttribute<double>(TEXT("DoubleAttribute"), DoubleValue, true, false);
+	TUniquePtr<IPCGAttributeAccessor> DoubleAccessor = PCGAttributeAccessorHelpers::CreateAccessor(DoubleAttribute, PointData->Metadata);
+	REQUIRE(DoubleAccessor);
+
+	// Single point with default entry to test default value access
+	PointData->SetNumPoints(1);
+
+	SECTION("Get FVector4 from double accessor with AllowBroadcastAndConstructible succeeds")
+	{
+		TUniquePtr<const IPCGAttributeAccessorKeys> Keys = PCGAttributeAccessorHelpers::CreateConstKeys(PointData, FPCGAttributePropertySelector());
+		REQUIRE(Keys);
+		FVector4 TempVector4 = FVector4::Zero();
+		REQUIRE(DoubleAccessor->Get(TempVector4, *Keys, EPCGAttributeAccessorFlags::AllowBroadcastAndConstructible));
+		REQUIRE_EQUAL(TempVector4, FVector4(DoubleValue, DoubleValue, DoubleValue, DoubleValue));
+	}
+
+	SECTION("GetRange FVector4 from double accessor with AllowBroadcastAndConstructible succeeds")
+	{
+		constexpr int32 NumElements = 128;
+		FRandomStream RandomStream(42);
+
+		PointData->SetNumPoints(NumElements);
+		TPCGValueRange<int64> MetadataEntryRange = PointData->GetMetadataEntryValueRange();
+		for (int32 i = 0; i < NumElements; ++i)
+		{
+			MetadataEntryRange[i] = PointData->Metadata->AddEntry();
+			DoubleAttribute->SetValue(MetadataEntryRange[i], static_cast<double>(RandomStream.FRand()) * 100.0);
+		}
+
+		TUniquePtr<const IPCGAttributeAccessorKeys> Keys = PCGAttributeAccessorHelpers::CreateConstKeys(PointData, FPCGAttributePropertySelector());
+		REQUIRE(Keys);
+		REQUIRE_EQUAL(Keys->GetNum(), NumElements);
+
+		TArray<FVector4> OutVector4s;
+		OutVector4s.SetNum(NumElements);
+
+		REQUIRE(DoubleAccessor->GetRange(TArrayView<FVector4>(OutVector4s), 0, *Keys, EPCGAttributeAccessorFlags::AllowBroadcastAndConstructible));
+
+		for (int32 i = 0; i < NumElements; ++i)
+		{
+			const double SourceDouble = DoubleAttribute->GetValueFromItemKey(MetadataEntryRange[i]);
+			REQUIRE_EQUAL(OutVector4s[i], FVector4(SourceDouble, SourceDouble, SourceDouble, SourceDouble));
+		}
+	}
+}
+
 // Tests that FSoftObjectPath and FSoftClassPath attributes can be read as FString even with StrictType flag (legacy special case), and that non-string type access fails with StrictType.
 TEST_CASE_METHOD(PCGTests::FPCGBaseTestWithBasicContext, "PCG::Accessor::Attribute::SoftObjectPathToStringWithStrictType", "[PCG][Accessor][Attribute]")
 {
